MemQuery - Updating the deprecated usage of ResidesInPermamnentPool

diff --git a/Engine/Plugins/MemoryUsageQueries/Source/MemoryUsageQueries/Private/MemoryUsageQueries.cpp b/Engine/Plugins/MemoryUsageQueries/Source/MemoryUsageQueries/Private/MemoryUsageQueries.cpp
--- a/Engine/Plugins/MemoryUsageQueries/Source/MemoryUsageQueries/Private/MemoryUsageQueries.cpp
+++ b/Engine/Plugins/MemoryUsageQueries/Source/MemoryUsageQueries/Private/MemoryUsageQueries.cpp
@@ -788,7 +788,8 @@ TArray<UObject*>& FMemoryUsageReferenceProcessor::GetRootSet()
 
 void FMemoryUsageReferenceProcessor::HandleTokenStreamObjectReference(UE::GC::FWorkerContext& Context, const UObject* ReferencingObject, UObject*& Object, UE::GC::FTokenId TokenIndex, const EGCTokenType TokenType, bool bAllowReferenceElimination)
 {
-	if (Object == nullptr || GUObjectArray.ObjectToIndex(Object) >= ReachableFull.Num() || GUObjectAllocator.ResidesInPermanentPool(Object) || GUObjectArray.IsDisregardForGC(Object))
+	FPermanentObjectPoolExtents PermanentPool;
+	if (Object == nullptr || GUObjectArray.ObjectToIndex(Object) >= ReachableFull.Num() || PermanentPool.Contains(Object) || GUObjectArray.IsDisregardForGC(Object))
 	{
 		return;
 	}
