Missing include causing non-unity build failure.

diff --git a/Engine/Plugins/MemoryUsageQueries/Source/MemoryUsageQueries/Private/MemoryUsageQueries.cpp b/Engine/Plugins/MemoryUsageQueries/Source/MemoryUsageQueries/Private/MemoryUsageQueries.cpp
--- a/Engine/Plugins/MemoryUsageQueries/Source/MemoryUsageQueries/Private/MemoryUsageQueries.cpp
+++ b/Engine/Plugins/MemoryUsageQueries/Source/MemoryUsageQueries/Private/MemoryUsageQueries.cpp
@@ -6,6 +6,7 @@
 
 #include "Engine/BlueprintGeneratedClass.h"
 #include "Misc/CoreMisc.h"
+#include "Misc/PackageName.h"
 #include "UObject/GCObject.h"
 #include "UObject/UObjectAllocator.h"
 #include "UObject/UObjectIterator.h"
