Fixed non-unity/pch compile errors

diff --git a/Engine/Plugins/MemoryUsageQueries/Source/MemoryUsageQueries/Private/MemoryUsageQueriesPrivate.h b/Engine/Plugins/MemoryUsageQueries/Source/MemoryUsageQueries/Private/MemoryUsageQueriesPrivate.h
--- a/Engine/Plugins/MemoryUsageQueries/Source/MemoryUsageQueries/Private/MemoryUsageQueriesPrivate.h
+++ b/Engine/Plugins/MemoryUsageQueries/Source/MemoryUsageQueries/Private/MemoryUsageQueriesPrivate.h
@@ -5,6 +5,7 @@
 #include "Containers/UnrealString.h"
 #include "CoreTypes.h"
 #include "Engine/EngineTypes.h"
+#include "MemoryUsageInfoProvider.h"
 #include "UObject/FastReferenceCollector.h"
 
 namespace MemoryUsageQueries::Internal
diff --git a/Engine/Source/Runtime/Engine/Public/PrimitiveComponentId.h b/Engine/Source/Runtime/Engine/Public/PrimitiveComponentId.h
--- a/Engine/Source/Runtime/Engine/Public/PrimitiveComponentId.h
+++ b/Engine/Source/Runtime/Engine/Public/PrimitiveComponentId.h
@@ -2,6 +2,8 @@
 
 #pragma once
 
+#include "Templates/TypeHash.h"
+
 /**
  * Class used to identify UPrimitiveComponents on the rendering thread without having to pass the pointer around,
  * Which would make it easy for people to access game thread variables from the rendering thread.
diff --git a/Engine/Source/Runtime/Renderer/Private/Strata/Strata.h b/Engine/Source/Runtime/Renderer/Private/Strata/Strata.h
--- a/Engine/Source/Runtime/Renderer/Private/Strata/Strata.h
+++ b/Engine/Source/Runtime/Renderer/Private/Strata/Strata.h
@@ -14,6 +14,7 @@
 
 // Forward declarations.
 class FScene;
+class FSceneRenderer;
 class FRDGBuilder;
 struct FMinimalSceneTextures;
 struct FScreenPassTexture;
diff --git a/Engine/Source/Runtime/Renderer/Public/ViewDebug.h b/Engine/Source/Runtime/Renderer/Public/ViewDebug.h
--- a/Engine/Source/Runtime/Renderer/Public/ViewDebug.h
+++ b/Engine/Source/Runtime/Renderer/Public/ViewDebug.h
@@ -3,6 +3,7 @@
 #pragma once
 
 #include "CoreMinimal.h"
+#include "Misc/ScopeRWLock.h"
 #include "PrimitiveSceneInfo.h"
 
 #if !UE_BUILD_SHIPPING // TODO: Decide whether or not the struct should be entirely stripped out of shipping
