MemQuery now supports CSV format by providing csv in addition to the command.
Cleaned up the output to be less noisy.

diff --git a/Engine/Plugins/MemoryUsageQueries/Source/MemoryUsageQueries/Private/MemoryUsageQueries.cpp b/Engine/Plugins/MemoryUsageQueries/Source/MemoryUsageQueries/Private/MemoryUsageQueries.cpp
--- a/Engine/Plugins/MemoryUsageQueries/Source/MemoryUsageQueries/Private/MemoryUsageQueries.cpp
+++ b/Engine/Plugins/MemoryUsageQueries/Source/MemoryUsageQueries/Private/MemoryUsageQueries.cpp
@@ -34,6 +34,7 @@ bool FMemoryUsageQueriesExec::Exec(UWorld* InWorld, const TCHAR* Cmd, FOutputDev
 	if (FParse::Command(&Cmd, TEXT("MemQuery")))
 	{
 		const bool bTruncate = !FParse::Param(Cmd, TEXT("notrunc"));
+		const bool bCSV = FParse::Param(Cmd, TEXT("csv"));
 
 		// Parse some common options
 		FString Name;
@@ -136,8 +137,7 @@ bool FMemoryUsageQueriesExec::Exec(UWorld* InWorld, const TCHAR* Cmd, FOutputDev
 
 				if (MemoryUsageQueries::GetDependenciesWithSize(MemoryUsageQueries::CurrentMemoryUsageInfoProvider, Name, DependenciesWithSize, &Ar))
 				{
-					Ar.Logf(TEXT("Dependencies (%s):"), *Name);
-					MemoryUsageQueries::Internal::PrintTagsWithSize(Ar, DependenciesWithSize, TEXT("dependencies"), bTruncate, Limit);
+					MemoryUsageQueries::Internal::PrintTagsWithSize(Ar, DependenciesWithSize, TEXT("dependencies"), bTruncate, Limit, bCSV);
 				}
 
 				return true;
@@ -153,8 +153,7 @@ bool FMemoryUsageQueriesExec::Exec(UWorld* InWorld, const TCHAR* Cmd, FOutputDev
 
 				if (MemoryUsageQueries::GetDependenciesWithSizeCombined(MemoryUsageQueries::CurrentMemoryUsageInfoProvider, Packages, CombinedDependenciesWithSize, &Ar))
 				{
-					Ar.Logf(TEXT("Combined dependencies (%s):"), *Names);
-					MemoryUsageQueries::Internal::PrintTagsWithSize(Ar, CombinedDependenciesWithSize, TEXT("combined dependencies"), bTruncate, Limit);
+					MemoryUsageQueries::Internal::PrintTagsWithSize(Ar, CombinedDependenciesWithSize, TEXT("combined dependencies"), bTruncate, Limit, bCSV);
 				}
 
 				return true;
@@ -170,8 +169,7 @@ bool FMemoryUsageQueriesExec::Exec(UWorld* InWorld, const TCHAR* Cmd, FOutputDev
 
 				if (MemoryUsageQueries::GetDependenciesWithSizeShared(MemoryUsageQueries::CurrentMemoryUsageInfoProvider, Packages, SharedDependenciesWithSize, &Ar))
 				{
-					Ar.Logf(TEXT("Shared dependencies (%s):"), *Names);
-					MemoryUsageQueries::Internal::PrintTagsWithSize(Ar, SharedDependenciesWithSize, TEXT("shared dependencies"), bTruncate, Limit);
+					MemoryUsageQueries::Internal::PrintTagsWithSize(Ar, SharedDependenciesWithSize, TEXT("shared dependencies"), bTruncate, Limit, bCSV);
 				}
 
 				return true;
@@ -187,8 +185,7 @@ bool FMemoryUsageQueriesExec::Exec(UWorld* InWorld, const TCHAR* Cmd, FOutputDev
 
 				if (MemoryUsageQueries::GetDependenciesWithSizeUnique(MemoryUsageQueries::CurrentMemoryUsageInfoProvider, Packages, UniqueDependenciesWithSize, &Ar))
 				{
-					Ar.Logf(TEXT("Unique dependencies (%s):"), *Names);
-					MemoryUsageQueries::Internal::PrintTagsWithSize(Ar, UniqueDependenciesWithSize, TEXT("unique dependencies"), bTruncate, Limit);
+					MemoryUsageQueries::Internal::PrintTagsWithSize(Ar, UniqueDependenciesWithSize, TEXT("unique dependencies"), bTruncate, Limit, bCSV);
 				}
 
 				return true;
@@ -204,8 +201,7 @@ bool FMemoryUsageQueriesExec::Exec(UWorld* InWorld, const TCHAR* Cmd, FOutputDev
 
 				if (MemoryUsageQueries::GetDependenciesWithSizeCommon(MemoryUsageQueries::CurrentMemoryUsageInfoProvider, Packages, CommonDependenciesWithSize, &Ar))
 				{
-					Ar.Logf(TEXT("Common dependencies (%s):"), *Names);
-					MemoryUsageQueries::Internal::PrintTagsWithSize(Ar, CommonDependenciesWithSize, TEXT("common dependencies"), bTruncate, Limit);
+					MemoryUsageQueries::Internal::PrintTagsWithSize(Ar, CommonDependenciesWithSize, TEXT("common dependencies"), bTruncate, Limit, bCSV);
 				}
 
 				return true;
@@ -246,8 +242,7 @@ bool FMemoryUsageQueriesExec::Exec(UWorld* InWorld, const TCHAR* Cmd, FOutputDev
 
 			if (bSuccess)
 			{
-				Ar.Logf(TEXT("%d largest assets:"), (Limit != -1 ? FMath::Min(Limit, AssetsWithSize.Num()) : AssetsWithSize.Num()));
-				MemoryUsageQueries::Internal::PrintTagsWithSize(Ar, AssetsWithSize, TEXT("largest assets"), bTruncate, Limit);
+				MemoryUsageQueries::Internal::PrintTagsWithSize(Ar, AssetsWithSize, TEXT("largest assets"), bTruncate, Limit, bCSV);
 			}
 
 			return true;
@@ -268,8 +263,7 @@ bool FMemoryUsageQueriesExec::Exec(UWorld* InWorld, const TCHAR* Cmd, FOutputDev
 
 			if (MemoryUsageQueries::GetFilteredClassesWithSize(ClassesWithSize, Group, AssetName, &Ar))
 			{
-				Ar.Logf(TEXT("%d largest classes"), (Limit != -1 ? FMath::Min(Limit, ClassesWithSize.Num()) : ClassesWithSize.Num()));
-				MemoryUsageQueries::Internal::PrintTagsWithSize(Ar, ClassesWithSize, *FString::Printf(TEXT("largest classes")), bTruncate, Limit);
+				MemoryUsageQueries::Internal::PrintTagsWithSize(Ar, ClassesWithSize, *FString::Printf(TEXT("largest classes")), bTruncate, Limit, bCSV);
 			}
 
 			return true;
@@ -290,8 +284,7 @@ bool FMemoryUsageQueriesExec::Exec(UWorld* InWorld, const TCHAR* Cmd, FOutputDev
 
 			if (MemoryUsageQueries::GetFilteredGroupsWithSize(GroupsWithSize, AssetName, Class, &Ar))
 			{
-				Ar.Logf(TEXT("%d largest groups"), (Limit != -1 ? FMath::Min(Limit, GroupsWithSize.Num()) : GroupsWithSize.Num()));
-				MemoryUsageQueries::Internal::PrintTagsWithSize(Ar, GroupsWithSize, *FString::Printf(TEXT("largest groups")), bTruncate, Limit);
+				MemoryUsageQueries::Internal::PrintTagsWithSize(Ar, GroupsWithSize, *FString::Printf(TEXT("largest groups")), bTruncate, Limit, bCSV);
 			}
 
 			return true;
@@ -342,7 +335,7 @@ bool FMemoryUsageQueriesExec::Exec(UWorld* InWorld, const TCHAR* Cmd, FOutputDev
 				}
 
 				PresetSavings.ValueSort(TGreater<uint64>());
-				MemoryUsageQueries::Internal::PrintTagsWithSize(Ar, PresetSavings, *FString::Printf(TEXT("possible savings")), bTruncate);
+				MemoryUsageQueries::Internal::PrintTagsWithSize(Ar, PresetSavings, *FString::Printf(TEXT("possible savings")), bTruncate, bCSV);
 
 				return true;
 			}
@@ -1089,12 +1082,11 @@ void RemoveFilteredPackages(TMap<FName, uint64>& OutPackagesWithSize, const FStr
 	}
 }
 
-void PrintTagsWithSize(FOutputDevice& Ar, const TMap<FName, uint64>& TagsWithSize, const TCHAR* Name, bool bTruncate /* = false */, int32 Limit /* = -1 */)
+void PrintTagsWithSize(FOutputDevice& Ar, const TMap<FName, uint64>& TagsWithSize, const TCHAR* Name, bool bTruncate /* = false */, int32 Limit /* = -1 */, bool bCSV /* = false */)
 {
 	uint64 TotalSize = 0U;
 	static const FString NoScopeString(TEXT("No scope"));
 
-	Ar.Logf(TEXT("----------------------------------------------------------"));
 	if (Limit < 0)
 	{
 		Limit = DefaultResultLimit;
@@ -1104,6 +1096,11 @@ void PrintTagsWithSize(FOutputDevice& Ar, const TMap<FName, uint64>& TagsWithSiz
 	int32 TagsToDisplay = bTruncate ? FMath::Min(Num, Limit) : Num;
 	int It = 0;
 
+	if (bCSV)
+	{
+		Ar.Logf(TEXT(",Name,SizeMB,SizeKB"));
+	}
+
 	for (auto& Elem : TagsWithSize)
 	{
 		if (It++ >= TagsToDisplay)
@@ -1114,18 +1111,36 @@ void PrintTagsWithSize(FOutputDevice& Ar, const TMap<FName, uint64>& TagsWithSiz
 		TotalSize += Elem.Value;
 		const FString KeyName = Elem.Key.IsValid() ? Elem.Key.ToString() : NoScopeString;
 
-		Ar.Logf(TEXT("%s - %.2f MB (%.2f KB)"), *KeyName, Elem.Value / (1024.f * 1024.f), Elem.Value / 1024.f);
-	}
+		const float SizeMB = Elem.Value / (1024.0f * 1024.0f);
+		const float SizeKB = Elem.Value / 1024.0f;
 
-	Ar.Logf(TEXT("----------------------------------------------------------"));
+		if (bCSV)
+		{
+			Ar.Logf(TEXT(",%s,%.2f,%.2f"), *KeyName, SizeMB, SizeKB);
+		}
+		else
+		{
+			Ar.Logf(TEXT("%s - %.2f MB (%.2f KB)"), *KeyName, SizeMB, SizeKB);
+		}
+	}
 
-	if (TagsToDisplay < Num)
+	if (TagsToDisplay < Num && !bCSV)
 	{
+		Ar.Logf(TEXT("----------------------------------------------------------"));
 		Ar.Logf(TEXT("<<truncated>> - displayed %d out of %d %s."), TagsToDisplay, Num, Name);
 	}
 
-	Ar.Logf(TEXT("TOTAL: %.2f MB (%.2f KB)"), TotalSize / (1024.f * 1024.f), TotalSize / 1024.f);
-	Ar.Logf(TEXT("----------------------------------------------------------"));
+	const float TotalSizeMB = TotalSize / (1024.0f * 1024.0f);
+	const float TotalSizeKB = TotalSize / 1024.0f;
+
+	if (bCSV)
+	{
+		Ar.Logf(TEXT(",TOTAL,%.2f,%.2f"), TotalSizeMB, TotalSizeKB);
+	}
+	else
+	{
+		Ar.Logf(TEXT("TOTAL: %.2f MB (%.2f KB)"), TotalSizeMB, TotalSizeKB);
+	}
 }
 
 } // namespace Internal
diff --git a/Engine/Plugins/MemoryUsageQueries/Source/MemoryUsageQueries/Private/MemoryUsageQueriesPrivate.h b/Engine/Plugins/MemoryUsageQueries/Source/MemoryUsageQueries/Private/MemoryUsageQueriesPrivate.h
--- a/Engine/Plugins/MemoryUsageQueries/Source/MemoryUsageQueries/Private/MemoryUsageQueriesPrivate.h
+++ b/Engine/Plugins/MemoryUsageQueries/Source/MemoryUsageQueries/Private/MemoryUsageQueriesPrivate.h
@@ -67,6 +67,6 @@ void GetPackagesSize(const IMemoryUsageInfoProvider* MemoryUsageInfoProvider, co
 void RemoveNonExistentPackages(TMap<FName, uint64>& OutPackagesWithSize);
 void RemoveFilteredPackages(TMap<FName, uint64>& OutPackagesWithSize, const FString& AssetSubstring);
 
-void PrintTagsWithSize(FOutputDevice& Ar, const TMap<FName, uint64>& TagsWithSize, const TCHAR* Name, bool bTruncate = false, int32 Limit = -1);
+void PrintTagsWithSize(FOutputDevice& Ar, const TMap<FName, uint64>& TagsWithSize, const TCHAR* Name, bool bTruncate = false, int32 Limit = -1, bool bCSV = false);
 
 }
