Memory Usage config provides a means of specifying a collection, a list of substring paths to match against asset package paths and an optional Budget in MB to compare total memory in the collection against.
Paths can specify includes and excludes which help narrow down the collection of interest.

diff --git a/Engine/Plugins/MemoryUsageQueries/Source/MemoryUsageQueries/Private/MemoryUsageInfoProvider.cpp b/Engine/Plugins/MemoryUsageQueries/Source/MemoryUsageQueries/Private/MemoryUsageInfoProvider.cpp
--- a/Engine/Plugins/MemoryUsageQueries/Source/MemoryUsageQueries/Private/MemoryUsageInfoProvider.cpp
+++ b/Engine/Plugins/MemoryUsageQueries/Source/MemoryUsageQueries/Private/MemoryUsageInfoProvider.cpp
@@ -31,6 +31,25 @@ uint64 FMemoryUsageInfoProviderLLM::GetAssetsMemoryUsage(const TSet<FName>& Asse
 	return 0U;
 }
 
+uint64 FMemoryUsageInfoProviderLLM::GetAssetsMemoryUsageWithSize(const TSet<FName>& Assets, TMap<FName, uint64>& OutSizes, FOutputDevice* ErrorOutput /* = GLog */) const
+{
+#if ENABLE_LOW_LEVEL_MEM_TRACKER
+	if (FLowLevelMemTracker::IsEnabled())
+	{
+		uint64 TotalSize = 0;
+		for (const FName& Asset : Assets)
+		{
+			int64 AssetSize = FLowLevelMemTracker::Get().GetTagAmountForTracker(ELLMTracker::Default, Asset, ELLMTagSet::Assets, false);
+			OutSizes.Add(Asset, AssetSize);
+			TotalSize += AssetSize;
+		}
+		return TotalSize;
+	}
+#endif
+	ErrorOutput->Logf(TEXT("MemoryUsageInfoProvider Error: LLM is disabled. Please run with -LLM"));
+	return 0U;
+}
+
 void FMemoryUsageInfoProviderLLM::GetAllAssetsWithSize(TMap<FName, uint64>& OutAssets, FOutputDevice* ErrorOutput /* = GLog */) const
 {
 	OutAssets.Reset();
diff --git a/Engine/Plugins/MemoryUsageQueries/Source/MemoryUsageQueries/Private/MemoryUsageQueries.cpp b/Engine/Plugins/MemoryUsageQueries/Source/MemoryUsageQueries/Private/MemoryUsageQueries.cpp
--- a/Engine/Plugins/MemoryUsageQueries/Source/MemoryUsageQueries/Private/MemoryUsageQueries.cpp
+++ b/Engine/Plugins/MemoryUsageQueries/Source/MemoryUsageQueries/Private/MemoryUsageQueries.cpp
@@ -8,6 +8,7 @@
 
 #include "Engine/BlueprintGeneratedClass.h"
 #include "Misc/PackageName.h"
+#include "Misc/WildcardString.h"
 #include "Modules/ModuleManager.h"
 #include "Templates/Greater.h"
 #include "UObject/UObjectAllocator.h"
@@ -28,6 +29,45 @@ struct FMemoryUsageQueriesExec : public FSelfRegisteringExec
 	virtual bool Exec(UWorld* Inworld, const TCHAR* Cmd, FOutputDevice& Ar) override;
 };
 
+struct FAssetMemoryBreakdown
+{
+	FAssetMemoryBreakdown() :
+		ExclusiveSize(0),
+		UniqueSize(0),
+		SharedSize(0),
+		TotalSize(0)
+	{
+
+	}
+
+	uint64 ExclusiveSize;
+	uint64 UniqueSize;
+	uint64 SharedSize;
+	uint64 TotalSize;
+};
+
+struct FAssetMemoryDetails
+{
+	FAssetMemoryDetails() :
+		UniqueRefCount(0),
+		SharedRefCount(0)
+	{
+	}
+
+	// Assets Package Name
+	FName PackageName;
+
+	FAssetMemoryBreakdown MemoryBreakdown;
+
+	// List of dependencies for this asset
+	TSet<FName> Dependencies;
+
+	TMap<FName, FAssetMemoryBreakdown> DependenciesToMemoryMap;
+
+	int32 UniqueRefCount;
+	int32 SharedRefCount;
+};
+
 static FMemoryUsageQueriesExec DebugSettinsSubsystemExecInstance;
 static int32 DefaultResultLimit = 15;
 
@@ -35,6 +75,8 @@ bool FMemoryUsageQueriesExec::Exec(UWorld* InWorld, const TCHAR* Cmd, FOutputDev
 {
 	if (FParse::Command(&Cmd, TEXT("MemQuery")))
 	{
+		Ar.Logf(TEXT("MemQuery: %s"), Cmd);
+
 		const bool bTruncate = !FParse::Param(Cmd, TEXT("notrunc"));
 		const bool bCSV = FParse::Param(Cmd, TEXT("csv"));
 
@@ -339,6 +381,236 @@ bool FMemoryUsageQueriesExec::Exec(UWorld* InWorld, const TCHAR* Cmd, FOutputDev
 				PresetSavings.ValueSort(TGreater<uint64>());
 				MemoryUsageQueries::Internal::PrintTagsWithSize(Ar, PresetSavings, *FString::Printf(TEXT("possible savings")), bTruncate, bCSV);
 
+				return true;
+			}
+		}
+		else if (FParse::Command(&Cmd, TEXT("Collection")))
+		{
+			const bool bShowDependencies = FParse::Param(Cmd, TEXT("ShowDeps"));
+
+			const UMemoryUsageQueriesConfig* Config = GetDefault<UMemoryUsageQueriesConfig>();
+			for (auto It = Config->Collections.CreateConstIterator(); It; ++It)
+			{
+				const FCollectionInfo& CollectionInfo = *It;
+				
+				if (!FParse::Command(&Cmd, *CollectionInfo.Name))
+				{
+					continue;
+				}
+				
+				// Retrieve a list of all assets that have allocations we are currently tracking.
+				TMap<FName, uint64> AssetsWithSize;
+				bool bSuccess = MemoryUsageQueries::GetFilteredPackagesWithSize(AssetsWithSize, NAME_None, "", NAME_None, &Ar);
+
+				if (!bSuccess)
+				{
+					Ar.Logf(TEXT("Failed to gather assets for Collection %s"), *CollectionInfo.Name);
+					return false;
+				}
+
+				// Will return true if the Package Name matches any of the conditions in the array of Paths
+				auto PackageNameMatches = [](const FString& PackageName, const TArray<FString>& Conditions)->bool {
+					for (const FString& Condition : Conditions)
+					{
+						if ((FWildcardString::ContainsWildcards(*Condition) && FWildcardString::IsMatch(*Condition, *PackageName)) || (PackageName.Contains(Condition)))
+						{
+							return true;
+						}
+					}
+
+					return false;
+				};
+				
+				// See if any of the asset paths match those of our matching paths and are valid
+				TArray<FString> PackageNames;
+
+				TMap<FName, FAssetMemoryDetails> AssetMemoryMap;
+				for (const TPair<FName, uint64>& AssetSizePair : AssetsWithSize)
+				{
+					const FString& PackageName = AssetSizePair.Key.ToString();
+					
+					if (!FPackageName::IsValidLongPackageName(PackageName))
+					{
+						continue;
+					}
+
+					// If path is Included and NOT Excluded, it's a valid asset to consider
+					if (PackageNameMatches(PackageName, CollectionInfo.Includes) && !PackageNameMatches(PackageName, CollectionInfo.Excludes))
+					{
+						PackageNames.Add(PackageName);
+						FAssetMemoryDetails& AssetMemory = AssetMemoryMap.Add(AssetSizePair.Key);
+						AssetMemory.MemoryBreakdown.ExclusiveSize = AssetSizePair.Value;
+
+						FName LongPackageName;
+						if (!MemoryUsageQueries::Internal::GetLongNameAndDependencies(PackageName, LongPackageName, AssetMemory.Dependencies, &Ar))
+						{
+							Ar.Logf(TEXT("Failed to get dependencies foro Asset %s"), *PackageName);
+						}
+					}
+				}
+
+				// Gather list of dependencies. Internal dependencies are confined only to the set of packages passed in.
+				// External are dependencies that have additional references outside the set of packages passed in.
+				TMap<FName, uint64> InternalDependencies;
+				TMap<FName, uint64> ExternalDependencies;
+				if (!MemoryUsageQueries::GatherDependenciesForPackages(MemoryUsageQueries::CurrentMemoryUsageInfoProvider, PackageNames, InternalDependencies, ExternalDependencies, &Ar))
+				{
+					Ar.Logf(TEXT("Failed to gather memory usage for dependencies in Collection %s"), *CollectionInfo.Name);
+					return false;
+				}
+
+				uint64 TotalCollectionSize = 0;
+
+				// Determine which category where each assets dependency should reside
+				for (TPair<FName, FAssetMemoryDetails>& Asset : AssetMemoryMap)
+				{
+					FAssetMemoryDetails& AssetMemory = Asset.Value;
+					FAssetMemoryBreakdown& AssetMemoryDetails = AssetMemory.MemoryBreakdown;
+
+					for (FName& Dep : Asset.Value.Dependencies)
+					{
+						// Don't want to count asset itself, plus some dependencies might refer to other assets in the map
+						if (AssetMemoryMap.Contains(Dep))
+						{
+							continue;
+						}
+
+						FAssetMemoryBreakdown DependencyMemory;
+						uint64* UniqueMemory = InternalDependencies.Find(Dep);
+						uint64* SharedMemory = ExternalDependencies.Find(Dep);
+						bool bRecordDependency = false;
+
+						if (UniqueMemory != nullptr && *UniqueMemory != 0)
+						{
+							DependencyMemory.UniqueSize = *UniqueMemory;
+							AssetMemoryDetails.UniqueSize += DependencyMemory.UniqueSize;
+							AssetMemory.UniqueRefCount++;
+							bRecordDependency = true;
+						}
+
+						if (SharedMemory != nullptr && *SharedMemory != 0)
+						{
+							DependencyMemory.SharedSize = *SharedMemory;
+							AssetMemoryDetails.SharedSize += DependencyMemory.SharedSize;
+							AssetMemory.SharedRefCount++;
+							bRecordDependency = true;
+						}
+
+						if (bRecordDependency)
+						{
+							AssetMemory.DependenciesToMemoryMap.Add(Dep, DependencyMemory);
+						}
+					}
+
+					AssetMemoryDetails.TotalSize = AssetMemoryDetails.ExclusiveSize + AssetMemoryDetails.UniqueSize;
+					TotalCollectionSize += AssetMemoryDetails.TotalSize;
+				}
+
+				// Sort by TotalSize
+				AssetMemoryMap.ValueSort([](const FAssetMemoryDetails& A, const FAssetMemoryDetails& B) {
+						return A.MemoryBreakdown.TotalSize > B.MemoryBreakdown.TotalSize;
+					});
+
+				if (bCSV)
+				{
+					Ar.Logf(TEXT(",Asset,Exclusive KiB,Unique Refs KiB,Unique Ref Count,Shared Refs KiB,Shared Ref Count,Total KiB"));
+				}
+				else
+				{
+					Ar.Logf(
+						TEXT(" %100s %20s %20s %15s %20s %15s %25s"),
+						TEXT("Asset"),
+						TEXT("Exclusive KiB"),
+						TEXT("Unique Refs KiB"),
+						TEXT("Unique Ref Count"),
+						TEXT("Shared Refs KiB"),
+						TEXT("Shared Ref Count"),
+						TEXT("Total KiB")
+					);
+				}
+
+				// Asset listing
+				for (TPair<FName, FAssetMemoryDetails>& Asset : AssetMemoryMap)
+				{
+					FAssetMemoryDetails& AssetMemory = Asset.Value;
+					FAssetMemoryBreakdown& AssetMemoryDetails = AssetMemory.MemoryBreakdown;
+
+					if (bCSV)
+					{
+						Ar.Logf(TEXT(",%s,%.2f,%.2f,%d,%.2f,%d,%.2f"),*Asset.Key.ToString(),
+							AssetMemoryDetails.ExclusiveSize / 1024.f, 
+							AssetMemoryDetails.UniqueSize / 1024.f, 
+							AssetMemory.UniqueRefCount, 
+							AssetMemoryDetails.SharedSize / 1024.f,
+							AssetMemory.SharedRefCount,
+							AssetMemoryDetails.TotalSize / 1024.f);
+					}
+					else
+					{
+						Ar.Logf(
+							TEXT(" %100s %20.2f %20.2f %15d %20.2f %15d %25.2f"),
+							*Asset.Key.ToString(),
+							AssetMemoryDetails.ExclusiveSize / 1024.f, 
+							AssetMemoryDetails.SharedSize / 1024.f,
+							AssetMemory.SharedRefCount,
+							AssetMemoryDetails.UniqueSize / 1024.f, 
+							AssetMemory.UniqueRefCount, 
+							AssetMemoryDetails.TotalSize / 1024.f
+						);
+					}
+				}
+
+				// Asset dependencies listing
+				if (bShowDependencies)
+				{
+					if (bCSV)
+					{
+						Ar.Logf(TEXT(",Asset,Dependency,Unique KiB,Shared KiB"));
+					}
+					else
+					{
+						Ar.Logf(TEXT(" %100s %100s %20s %20s"),
+							TEXT("Asset"),
+							TEXT("Dependency"),
+							TEXT("Unique KiB"),
+							TEXT("Shared KiB")
+						);
+					}
+
+					for (TPair<FName, FAssetMemoryDetails>& Asset : AssetMemoryMap)
+					{
+						for (TPair<FName, FAssetMemoryBreakdown>& Dep : Asset.Value.DependenciesToMemoryMap)
+						{
+							const FString DependencyAssetName = Dep.Key.ToString();
+							const FAssetMemoryBreakdown& DepedencyMemoryDetails = Dep.Value;
+
+							if (bCSV)
+							{
+								Ar.Logf(TEXT(",%s,%s,%.2f,%.2f"), *Asset.Key.ToString(), *DependencyAssetName,
+									DepedencyMemoryDetails.UniqueSize / 1024.f, DepedencyMemoryDetails.SharedSize / 1024.f);
+							}
+							else
+							{
+								Ar.Logf(TEXT(" %100s %100s %20.2f %20.2f"), 
+									*Asset.Key.ToString(), 
+									*DependencyAssetName,
+									DepedencyMemoryDetails.UniqueSize / 1024.f, 
+									DepedencyMemoryDetails.SharedSize / 1024.f
+								);
+							}
+						}
+					}
+				}
+
+				if (bCSV)
+				{
+					Ar.Logf(TEXT(",TOTAL KiB,%.2f"), TotalCollectionSize / 1024.f);
+				}
+				else
+				{
+					Ar.Logf(TEXT("TOTAL KiB: %.2f"), TotalCollectionSize / 1024.f);
+				}
+
 				return true;
 			}
 		}
@@ -464,7 +736,16 @@ void RegisterConsoleAutoCompleteEntries(TArray<FAutoCompleteCommand>& AutoComple
 	{
 		FAutoCompleteCommand AutoCompleteCommand;
 		AutoCompleteCommand.Command = FString::Printf(TEXT("MemQuery Savings %s"), *It.Key());
-		AutoCompleteCommand.Desc = FString::Printf(TEXT("Limit=<n> Lists potential savings among %s. How much memory can be saved it we delete certain object?"), *It.Key());
+		AutoCompleteCommand.Desc = FString::Printf(TEXT("Limit=<n> Lists potential savings among %s. How much memory can be saved it we delete certain object."), *It.Key());
+		AutoCompleteCommand.Color = ConsoleSettings->AutoCompleteCommandColor;
+		AutoCompleteList.Add(MoveTemp(AutoCompleteCommand));
+	}
+
+	for (auto It = MemQueryConfig->Collections.CreateConstIterator(); It; ++It)
+	{
+		FAutoCompleteCommand AutoCompleteCommand;
+		AutoCompleteCommand.Command = FString::Printf(TEXT("MemQuery Collection %s"), *It->Name);
+		AutoCompleteCommand.Desc = TEXT("Lists memory used by a collection. Can show dependency breakdown. [-csv, -showdeps]");
 		AutoCompleteCommand.Color = ConsoleSettings->AutoCompleteCommandColor;
 		AutoCompleteList.Add(MoveTemp(AutoCompleteCommand));
 	}
@@ -542,6 +823,28 @@ bool GetMemoryUsageCommon(const IMemoryUsageInfoProvider* MemoryUsageInfoProvide
 	return true;
 }
 
+bool GatherDependenciesForPackages(const IMemoryUsageInfoProvider* MemoryUsageInfoProvider, const TArray<FString>& PackageNames,
+	TMap<FName, uint64>& OutInternalDeps, TMap<FName, uint64>& OutExternalDeps, FOutputDevice* ErrorOutput)
+{
+	TSet<FName> RemovablePackages;
+	if (!Internal::GetRemovablePackages(PackageNames, RemovablePackages, ErrorOutput))
+	{
+		return false;
+	}
+
+	TSet<FName> UnremovablePackages;
+	if (!Internal::GetUnremovablePackages(PackageNames, UnremovablePackages, ErrorOutput))
+	{
+		return false;
+	}
+
+	MemoryUsageInfoProvider->GetAssetsMemoryUsageWithSize(RemovablePackages, OutInternalDeps, ErrorOutput);
+	MemoryUsageInfoProvider->GetAssetsMemoryUsageWithSize(UnremovablePackages, OutExternalDeps, ErrorOutput);
+
+	return true;
+}
+
+
 bool GetDependenciesWithSize(const IMemoryUsageInfoProvider* MemoryUsageInfoProvider, const FString& PackageName, TMap<FName, uint64>& OutDependenciesWithSize, FOutputDevice* ErrorOutput /* = GLog */)
 {
 	FName LongPackageName;
diff --git a/Engine/Plugins/MemoryUsageQueries/Source/MemoryUsageQueries/Public/MemoryUsageInfoProvider.h b/Engine/Plugins/MemoryUsageQueries/Source/MemoryUsageQueries/Public/MemoryUsageInfoProvider.h
--- a/Engine/Plugins/MemoryUsageQueries/Source/MemoryUsageQueries/Public/MemoryUsageInfoProvider.h
+++ b/Engine/Plugins/MemoryUsageQueries/Source/MemoryUsageQueries/Public/MemoryUsageInfoProvider.h
@@ -14,6 +14,7 @@ public:
 	virtual uint64 GetAssetMemoryUsage(FName Asset, FOutputDevice* ErrorOutput) const = 0;
 	virtual uint64 GetAssetsMemoryUsage(const TSet<FName>& Assets, FOutputDevice* ErrorOutput) const = 0;
 	virtual void GetAllAssetsWithSize(TMap<FName, uint64>& OutAssets, FOutputDevice* ErrorOutput) const = 0;
+	virtual uint64 GetAssetsMemoryUsageWithSize(const TSet<FName>& Assets, TMap<FName, uint64>& OutSizes, FOutputDevice* ErrorOutput = GLog) const = 0;
 };
 
 class MEMORYUSAGEQUERIES_API FMemoryUsageInfoProviderLLM : public IMemoryUsageInfoProvider
@@ -22,6 +23,8 @@ public:
 	virtual uint64 GetAssetMemoryUsage(FName Asset, FOutputDevice* ErrorOutput = GLog) const override;
 	virtual uint64 GetAssetsMemoryUsage(const TSet<FName>& Assets, FOutputDevice* ErrorOutput = GLog) const override;
 	virtual void GetAllAssetsWithSize(TMap<FName, uint64>& OutAssets, FOutputDevice* ErrorOutput = GLog) const override;
+	virtual uint64 GetAssetsMemoryUsageWithSize(const TSet<FName>& Assets, TMap<FName, uint64>& OutSizes, FOutputDevice* ErrorOutput = GLog) const override;
+
 
 #if ENABLE_LOW_LEVEL_MEM_TRACKER
 	virtual void GetFilteredTagsWithSize(TMap<FName, uint64>& OutTags, ELLMTracker Tracker, ELLMTagSet TagSet, TArray<FLLMTagSetAllocationFilter>& Filters, FOutputDevice* ErrorOutput = GLog) const;
diff --git a/Engine/Plugins/MemoryUsageQueries/Source/MemoryUsageQueries/Public/MemoryUsageQueries.h b/Engine/Plugins/MemoryUsageQueries/Source/MemoryUsageQueries/Public/MemoryUsageQueries.h
--- a/Engine/Plugins/MemoryUsageQueries/Source/MemoryUsageQueries/Public/MemoryUsageQueries.h
+++ b/Engine/Plugins/MemoryUsageQueries/Source/MemoryUsageQueries/Public/MemoryUsageQueries.h
@@ -21,6 +21,8 @@ MEMORYUSAGEQUERIES_API bool GetMemoryUsageCombined(const IMemoryUsageInfoProvide
 MEMORYUSAGEQUERIES_API bool GetMemoryUsageShared(const IMemoryUsageInfoProvider* MemoryUsageInfoProvider, const TArray<FString>& PackageNames, uint64& OutTotalSize, FOutputDevice* ErrorOutput = GLog);
 MEMORYUSAGEQUERIES_API bool GetMemoryUsageUnique(const IMemoryUsageInfoProvider* MemoryUsageInfoProvider, const TArray<FString>& PackageNames, uint64& OutUniqueSize, FOutputDevice* ErrorOutput = GLog);
 MEMORYUSAGEQUERIES_API bool GetMemoryUsageCommon(const IMemoryUsageInfoProvider* MemoryUsageInfoProvider, const TArray<FString>& PackageNames, uint64& OutCommonSize, FOutputDevice* ErrorOutput = GLog);
+MEMORYUSAGEQUERIES_API bool GatherDependenciesForPackages(const IMemoryUsageInfoProvider* MemoryUsageInfoProvider, const TArray<FString>& PackageNames,
+	TMap<FName, uint64>& OutInternalDeps, TMap<FName, uint64>& OutExternalDeps, FOutputDevice* ErrorOutput = GLog);
 						    
 MEMORYUSAGEQUERIES_API bool GetDependenciesWithSize(const IMemoryUsageInfoProvider* MemoryUsageInfoProvider, const FString& PackageName, TMap<FName, uint64>& OutDependenciesWithSize, FOutputDevice* ErrorOutput = GLog);
 MEMORYUSAGEQUERIES_API bool GetDependenciesWithSizeCombined(const IMemoryUsageInfoProvider* MemoryUsageInfoProvider, const TArray<FString>& PackageNames, TMap<FName, uint64>& OutDependenciesWithSize, FOutputDevice* ErrorOutput = GLog);
diff --git a/Engine/Plugins/MemoryUsageQueries/Source/MemoryUsageQueries/Public/MemoryUsageQueriesConfig.h b/Engine/Plugins/MemoryUsageQueries/Source/MemoryUsageQueries/Public/MemoryUsageQueriesConfig.h
--- a/Engine/Plugins/MemoryUsageQueries/Source/MemoryUsageQueries/Public/MemoryUsageQueriesConfig.h
+++ b/Engine/Plugins/MemoryUsageQueries/Source/MemoryUsageQueries/Public/MemoryUsageQueriesConfig.h
@@ -4,13 +4,56 @@
 
 #include "MemoryUsageQueriesConfig.generated.h"
 
-// “MemQuery Savings” command is supposed to list potential memory savings for assets based on predefined presets. 
-// You can define them in DefaultMemoryUsageQueries.ini in your project – you have to provide preset name and UClass.
-// "MemQuery Savings" command should list all derived classes with potential savings for the specified preset.
+// “MemQuery Collection” will list memory consumption based on assets discovered that have a path
+// matching the one described in the ini file.
+// Matching supports Wildcards. e.g. */Game*/*
+// This will match assets such as /SomePath/Game/SomeLevelAsset
+// 
+// Define these in DefaultMemoryUsageQueries.ini in your project – you have to provide preset name and UClass.
+// 
+// e.g. Include all Characters and Pawns but exclude anything with Item in the path. Demonstrates wildcard usage.
+// [/Script/MemoryUsageQueries.MemoryUsageQueriesConfig]
+// +Collections = (Name="Pawns", Includes=("*/Character*/*", "/Pawn/"), Excludes=("*/Item*/*"))
+//
+// A budget can be defined which will compare the total memory against the budget
+// e.g.
+// [/Script/MemoryUsageQueries.MemoryUsageQueriesConfig]
+// +Collections = (Name="Pawns", Includes=("/Character/", "/Pawn/"), BudgetMB=2.0)
+//
+// "MemQuery Savings" command will list all derived classes with potential savings for the specified preset.
 // 
 // e.g.
 // [/Script/MemoryUsageQueries.MemoryUsageQueriesConfig]
 // +SavingsPresets = (("Pawns", "/Script/OurProject.OurCharacter"))
+//
+
+
+USTRUCT()
+struct FCollectionInfo
+{
+	GENERATED_BODY()
+
+	FCollectionInfo() :
+		BudgetMB(0.0f)
+	{
+	}
+
+	// Name of the collection
+	UPROPERTY()
+	FString Name;
+
+	// Collection of substrings to include whne matching against asset package paths
+	UPROPERTY()
+	TArray<FString> Includes;
+
+	// Paths of asset package paths to exclude from the results
+	UPROPERTY()
+	TArray<FString> Excludes;
+
+	// Budget in MB
+	UPROPERTY()
+	float BudgetMB;
+};
 
 UCLASS(config = MemoryUsageQueries)
 class UMemoryUsageQueriesConfig : public UObject
@@ -18,6 +61,9 @@ class UMemoryUsageQueriesConfig : public UObject
 	GENERATED_BODY()
 
 public:
+	UPROPERTY(config)
+	TArray<FCollectionInfo> Collections;
+
 	UPROPERTY(config)
 	TMap<FString, FString> SavingsPresets;
 };
\ No newline at end of file
