Build compilation error when LLM is not enabled.

diff --git a/Engine/Plugins/MemoryUsageQueries/Source/MemoryUsageQueries/Private/MemoryUsageQueries.cpp b/Engine/Plugins/MemoryUsageQueries/Source/MemoryUsageQueries/Private/MemoryUsageQueries.cpp
--- a/Engine/Plugins/MemoryUsageQueries/Source/MemoryUsageQueries/Private/MemoryUsageQueries.cpp
+++ b/Engine/Plugins/MemoryUsageQueries/Source/MemoryUsageQueries/Private/MemoryUsageQueries.cpp
@@ -384,6 +384,7 @@ bool FMemoryUsageQueriesExec::Exec(UWorld* InWorld, const TCHAR* Cmd, FOutputDev
 				return true;
 			}
 		}
+#if ENABLE_LOW_LEVEL_MEM_TRACKER
 		else if (FParse::Command(&Cmd, TEXT("Collection")))
 		{
 			const bool bShowDependencies = FParse::Param(Cmd, TEXT("ShowDeps"));
@@ -614,6 +615,7 @@ bool FMemoryUsageQueriesExec::Exec(UWorld* InWorld, const TCHAR* Cmd, FOutputDev
 				return true;
 			}
 		}
+#endif
 	}
 
 	return false;
