MemQuery use PackageStore to resolve dependencies

#tests replay

diff --git a/Engine/Plugins/MemoryUsageQueries/Source/MemoryUsageQueries/Private/MemoryUsageQueries.cpp b/Engine/Plugins/MemoryUsageQueries/Source/MemoryUsageQueries/Private/MemoryUsageQueries.cpp
--- a/Engine/Plugins/MemoryUsageQueries/Source/MemoryUsageQueries/Private/MemoryUsageQueries.cpp
+++ b/Engine/Plugins/MemoryUsageQueries/Source/MemoryUsageQueries/Private/MemoryUsageQueries.cpp
@@ -1,12 +1,16 @@
 // Copyright Epic Games, Inc. All Rights Reserved.
 
 #include "MemoryUsageQueries.h"
-#include "AssetRegistry/AssetRegistryModule.h"
-#include "MemoryUsageQueriesConfig.h"
+
 #include "ConsoleSettings.h"
+#include "IPlatformFilePak.h"
+#include "MemoryUsageQueriesConfig.h"
 #include "MemoryUsageQueriesPrivate.h"
 
+#include "AssetRegistry/AssetData.h"
+#include "AssetRegistry/AssetRegistryModule.h"
 #include "Engine/BlueprintGeneratedClass.h"
+#include "IO/PackageStore.h"
 #include "Misc/PackageName.h"
 #include "Misc/WildcardString.h"
 #include "Modules/ModuleManager.h"
@@ -1137,27 +1141,411 @@ bool FMemoryUsageReferenceProcessor::GetUnreachablePackages(TSet<FName>& OutUnre
 	return true;
 }
 
-bool GetLongName(const FString& ShortPackageName, FName& OutLongPackageName, FOutputDevice* ErrorOutput /* = GLog */)
+static FAutoConsoleVariable CVarMemQueryUsePackageStore(
+	TEXT("MemQuery.UsePackageStore"), true,
+	TEXT("True - use PackageStore, false - use AssetRegistry."), ECVF_Default);
+
+static FAssetRegistryModule& GetAssetRegistryModule()
 {
-	FAssetRegistryModule& AssetRegistryModule = GetAssetRegistryModule();
+	FAssetRegistryModule& AssetRegistryModule = FModuleManager::Get().LoadModuleChecked<FAssetRegistryModule>(TEXT("AssetRegistry"));
 
-	if (FPackageName::IsValidLongPackageName(ShortPackageName))
+	if (IsInGameThread())
 	{
-		OutLongPackageName = FName(*ShortPackageName);
+		AssetRegistryModule.Get().WaitForCompletion();
 	}
-	else
+
+	return AssetRegistryModule;
+}
+
+class FPackageStoreLazyDatabase final
+{
+	FPackageStoreLazyDatabase()
+	{
+		FPackageName::OnContentPathMounted().AddRaw(this, &FPackageStoreLazyDatabase::OnContentPathMounted);
+		FPackageName::OnContentPathDismounted().AddRaw(this, &FPackageStoreLazyDatabase::OnContentPathDismounted);
+	}
+
+	~FPackageStoreLazyDatabase()
+	{
+		FPackageName::OnContentPathMounted().RemoveAll(this);
+		FPackageName::OnContentPathDismounted().RemoveAll(this);
+	}
+
+	void ResetDatabase()
+	{
+		bIsAssetDatabaseSearched = false;
+		bIsDirectoryIndexSearched = false;
+		Names.Empty();
+		Database.Empty();
+	}
+
+	void OnContentPathMounted(const FString&, const FString&) { ResetDatabase(); }
+	void OnContentPathDismounted(const FString&, const FString&) { ResetDatabase(); }
+
+	bool BuildDatabaseWhile(TFunctionRef<bool(const FPackageId& PackageId, const FName& PackageName)> Predicate)
+	{
+		bool bIterationBrokeEarly = false;
+
+		if (!bIsAssetDatabaseSearched)
+		{
+			const FAssetRegistryModule& AssetRegistryModule = GetAssetRegistryModule();
+			constexpr bool bOnlyOnDiskAssets = true;
+
+			AssetRegistryModule.Get().EnumerateAllAssets(
+				[&](const FAssetData& Data)
+				{
+					FPackageId PackageId = FPackageId::FromName(Data.PackageName);
+					{
+						LLM_SCOPE(TEXT("MemoryUsageQueries"));
+						Database.FindOrAdd(PackageId, Data.PackageName);
+						Names.Add(Data.PackageName);
+					}
+
+					if (!Invoke(Predicate, PackageId, Data.PackageName))
+					{
+						bIterationBrokeEarly = true;
+						return false;
+					}
+
+					return true;
+				},
+				bOnlyOnDiskAssets);
+
+			if (bIterationBrokeEarly)
+			{
+				return false;
+			}
+
+			bIsAssetDatabaseSearched = true;
+		}
+
+		if (!bIsDirectoryIndexSearched)
+		{
+			FPakPlatformFile::ForeachPackageInIostoreWhile(
+				[&](FName PackageName)
+				{
+					FPackageId PackageId = FPackageId::FromName(PackageName);
+					{
+						LLM_SCOPE(TEXT("MemoryUsageQueries"));
+						Database.FindOrAdd(PackageId, PackageName);
+						Names.Add(PackageName);
+					}
+
+					if (!Invoke(Predicate, PackageId, PackageName))
+					{
+						bIterationBrokeEarly = true;
+						return false;
+					}
+
+					return true;
+				});
+
+			if (bIterationBrokeEarly)
+			{
+				return false;
+			}
+
+			bIsDirectoryIndexSearched = true;
+		}
+
+		return true;
+	}
+
+	TMap<FPackageId, FName> Database;
+	TArray<FName> Names;
+	bool bIsAssetDatabaseSearched = false;
+	bool bIsDirectoryIndexSearched = false;
+
+public:
+	static FPackageStoreLazyDatabase& Get()
+	{
+		static FPackageStoreLazyDatabase Instance{};
+		return Instance;
+	}
+
+	// Blocking Call, builds full database
+	void IterateAllPackages(TFunctionRef<void(const FPackageId&)> Visitor)
 	{
-		OutLongPackageName = AssetRegistryModule.Get().GetFirstPackageByName(ShortPackageName);
-		if (OutLongPackageName == NAME_None)
+		BuildDatabaseWhile([](const auto&, const auto&) { return true; });
+
+		for (const TTuple<FPackageId, FName>& Pair : Database)
+		{
+			Invoke(Visitor, Pair.Key);
+		}
+	}
+
+
+	bool GetPackageNameFromId(FPackageId InPackageId, FName& OutFullName)
+	{
+		if (const FName* Found = Database.Find(InPackageId))
+		{
+			OutFullName = *Found;
+			return true;
+		}
+
+		const bool bIterationCompletedWithoutBreak = BuildDatabaseWhile(
+			[&](const FPackageId& PackageId, const FName& PackageName)
+			{
+				if (PackageId == InPackageId)
+				{
+					OutFullName = PackageName;
+					return false;
+				}
+
+				return true;
+			});
+
+		return !bIterationCompletedWithoutBreak;
+	}
+
+	bool GetFirstPackageNameFromPartialName(FStringView InPartialName, FName& OutPackageName)
+	{
+		auto SearchCriteria = [InPartialName](const FName& PackageName) -> bool
+		{
+			TCHAR Storage[FName::StringBufferSize];
+			PackageName.ToString(Storage);
+			return UE::String::FindFirst(Storage, InPartialName, ESearchCase::IgnoreCase) != INDEX_NONE;
+		};
+
+		if (const FName* Found = Names.FindByPredicate(SearchCriteria))
+		{
+			OutPackageName = *Found;
+			return true;
+		}
+
+		const bool bIterationCompletedWithoutBreak = BuildDatabaseWhile(
+			[&](const FPackageId&, const FName& PackageName)
+			{
+				if (SearchCriteria(PackageName))
+				{
+					OutPackageName = PackageName;
+					return false;
+				}
+
+				return true;
+			});
+
+		return !bIterationCompletedWithoutBreak;
+	}
+
+	bool DoesPackageExists(const FName& PackageName)
+	{
+		BuildDatabaseWhile([](const auto&, const auto&) { return true; });
+		return Names.Contains(PackageName);
+	}
+};
+
+class FPackageDependenciesLazyDatabase final
+{
+	FPackageDependenciesLazyDatabase()
+	{
+		FPackageName::OnContentPathMounted().AddRaw(this, &FPackageDependenciesLazyDatabase::OnContentPathMounted);
+		FPackageName::OnContentPathDismounted().AddRaw(this, &FPackageDependenciesLazyDatabase::OnContentPathDismounted);
+	}
+
+	~FPackageDependenciesLazyDatabase()
+	{
+		FPackageName::OnContentPathMounted().RemoveAll(this);
+		FPackageName::OnContentPathDismounted().RemoveAll(this);
+	}
+
+	void ResetDatabase()
+	{
+		Dependencies.Empty();
+		Referencers.Empty();
+		Leafs.Empty();
+		Roots.Empty();
+	}
+
+	void OnContentPathMounted(const FString&, const FString&) { ResetDatabase(); }
+	void OnContentPathDismounted(const FString&, const FString&) { ResetDatabase(); }
+
+	bool InsertPackage(const FPackageId RootPackageId)
+	{
+		LLM_SCOPE(TEXT("MemoryUsageQueries"));
+
+		TArray<FPackageId, TInlineAllocator<2048>> Queue;
+
+		Queue.Add(RootPackageId);
+		bool bAddedSuccessfully = false;
+
+		while (!Queue.IsEmpty())
+		{
+			const FPackageId PackageId = Queue.Top();
+			Queue.RemoveAt(0, 1, false);
+
+			if (Dependencies.Contains(PackageId) || Leafs.Contains(PackageId))
+			{
+				bAddedSuccessfully = true;
+				continue;
+			}
+
+			FPackageStoreEntry PackageEntry;
+			const EPackageStoreEntryStatus Status = FPackageStore::Get().GetPackageStoreEntry(PackageId, PackageEntry);
+			if (Status == EPackageStoreEntryStatus::Ok)
+			{
+				// add package dependencies
+				const TArrayView<const FPackageId> ImportedPackageIds = PackageEntry.ImportedPackageIds;
+				for (FPackageId DependentId : ImportedPackageIds)
+				{
+					Dependencies.FindOrAdd(PackageId).Add(DependentId);
+					Referencers.FindOrAdd(DependentId).Add(PackageId);
+					bAddedSuccessfully = true;
+				}
+				Queue.Append(ImportedPackageIds);
+
+#if WITH_EDITOR
+				// Add editor optional dependencies
+				const TArrayView<const FPackageId> OptionalImportedPackageIds = PackageEntry.OptionalSegmentImportedPackageIds;
+				for (FPackageId DependentId : OptionalImportedPackageIds)
+				{
+					Dependencies.FindOrAdd(PackageId).Add(DependentId);
+					Referencers.FindOrAdd(DependentId).Add(PackageId);
+					bAddedSuccessfully = true;
+				}
+				Queue.Append(OptionalImportedPackageIds);
+#endif
+
+
+				// Add leafs
+				if (ImportedPackageIds.IsEmpty()
+#if WITH_EDITOR
+					&& PackageEntry.UncookedPackageName.IsNone()
+					&& OptionalImportedPackageIds.IsEmpty()
+#endif
+				)
+				{
+					Leafs.Add(PackageId);
+					bAddedSuccessfully = true;
+				}
+			}
+		}
+
+		return bAddedSuccessfully;
+	}
+
+	TMap<FPackageId, TSet<FPackageId>> Dependencies;
+	TMap<FPackageId, TSet<FPackageId>> Referencers;
+	TSet<FPackageId> Leafs;
+	TSet<FPackageId> Roots;
+
+public:
+	static FPackageDependenciesLazyDatabase& Get()
+	{
+		static FPackageDependenciesLazyDatabase Instance{};
+		return Instance;
+	}
+
+	bool GetDependencies(const FPackageId PackageId, TSet<FPackageId>& OutDependencies)
+	{
+		if (Leafs.Contains(PackageId))
+		{
+			return true;
+		}
+
+		const TSet<FPackageId>* ChildrenSet = Dependencies.Find(PackageId);
+		if (ChildrenSet == nullptr)
+		{
+			FPackageStoreReadScope _(FPackageStore::Get());
+			if (!InsertPackage(PackageId))
+			{
+				return false;
+			}
+			ChildrenSet = Dependencies.Find(PackageId);
+		}
+		if (ChildrenSet == nullptr)
 		{
-			ErrorOutput->Logf(TEXT("MemQuery Error: Package not found: %s"), *ShortPackageName);
 			return false;
 		}
+
+		TArray<FPackageId, TInlineAllocator<2048>> Children(ChildrenSet->Array());
+		while (!Children.IsEmpty())
+		{
+			FPackageId Child = Children.Top();
+			Children.RemoveAt(0, 1, false);
+			if (!OutDependencies.Contains(Child))
+			{
+				OutDependencies.Add(Child);
+				if (const TSet<FPackageId>* GrandChildrenSet = Dependencies.Find(Child))
+				{
+					Children.Append(GrandChildrenSet->Array());
+				}
+			}
+		}
+
+		return true;
+	}
+
+	bool GetReferencers(const FPackageId InPackageId, TSet<FPackageId>& OutReferencers)
+	{
+		if (Roots.Contains(InPackageId))
+		{
+			return true;
+		}
+
+		TSet<FPackageId>* ParentsSet = Referencers.Find(InPackageId);
+		if (ParentsSet == nullptr)
+		{
+			FPackageStoreReadScope _(FPackageStore::Get());
+			// First IterateAllPackages call builds dependencies map
+			FPackageStoreLazyDatabase::Get().IterateAllPackages([this](const FPackageId& PackageId) { InsertPackage(PackageId); });
+			// Second IterateAllPackages call caches root nodes (only valid after full database is built)
+			FPackageStoreLazyDatabase::Get().IterateAllPackages(
+				[this](const FPackageId& PackageId)
+				{
+					if (!Referencers.Contains(PackageId))
+					{
+						Roots.Add(PackageId);
+					}
+				});
+
+			ParentsSet = Referencers.Find(InPackageId);
+		}
+
+		OutReferencers = MoveTemp(*ParentsSet);
+
+		return true;
+	}
+};
+
+static bool GetLongNamePackageStore(FStringView InShortPackageName, FName& OutLongPackageName)
+{
+	return FPackageStoreLazyDatabase::Get().GetFirstPackageNameFromPartialName(InShortPackageName, OutLongPackageName);
+}
+
+static bool GetLongNameAssetRegistry(FStringView InShortPackageName, FName& OutLongPackageName)
+{
+	const FAssetRegistryModule& AssetRegistryModule = GetAssetRegistryModule();
+
+	OutLongPackageName = AssetRegistryModule.Get().GetFirstPackageByName(InShortPackageName);
+	if (OutLongPackageName == NAME_None)
+	{
+		return false;
 	}
 
 	return true;
 }
 
+bool GetLongName(FStringView ShortPackageName, FName& OutLongPackageName, FOutputDevice* ErrorOutput /* = GLog */)
+{
+	if (FPackageName::IsValidLongPackageName(ShortPackageName))
+	{
+		OutLongPackageName = FName(ShortPackageName);
+		return true;
+	}
+
+	const bool Result =
+		CVarMemQueryUsePackageStore->GetBool() && FPackageStore::Get().HasAnyBackendsMounted()
+		? GetLongNamePackageStore(ShortPackageName, OutLongPackageName)
+		: GetLongNameAssetRegistry(ShortPackageName, OutLongPackageName);
+
+	if (!Result && ErrorOutput)
+	{
+		ErrorOutput->Logf(ELogVerbosity::Error, TEXT("MemQuery Error: Package not found: %.*s"), ShortPackageName.Len(), ShortPackageName.GetData());
+	}
+	return Result;
+}
+
 bool GetLongNames(const TArray<FString>& PackageNames, TSet<FName>& OutLongPackageNames, FOutputDevice* ErrorOutput /* = GLog */)
 {
 	for (const auto& Package : PackageNames)
@@ -1173,16 +1561,14 @@ bool GetLongNames(const TArray<FString>& PackageNames, TSet<FName>& OutLongPacka
 	return true;
 }
 
-bool GetLongNameAndDependencies(const FString& PackageName, FName& OutLongPackageName, TSet<FName>& OutDependencies, FOutputDevice* ErrorOutput /* = GLog */)
+bool GetLongNameAndDependencies(FStringView PackageName, FName& OutLongPackageName, TSet<FName>& OutDependencies, FOutputDevice* ErrorOutput /* = GLog */)
 {
 	if (!GetLongName(PackageName, OutLongPackageName, ErrorOutput))
 	{
 		return false;
 	}
 
-	FAssetRegistryModule& AssetRegistryModule = GetAssetRegistryModule();
-
-	GetTransitiveDependencies(OutLongPackageName, AssetRegistryModule, OutDependencies);
+	GetTransitiveDependencies(OutLongPackageName, OutDependencies);
 	OutDependencies.Add(OutLongPackageName);
 
 	return true;
@@ -1305,25 +1691,15 @@ bool GetUnremovablePackages(const TArray<FString>& PackagesToUnload, TSet<FName>
 	return true;
 }
 
-FAssetRegistryModule& GetAssetRegistryModule()
+static void GetTransitiveDependenciesAssetRegistry(FName PackageName, TSet<FName>& OutDependencies)
 {
-	FAssetRegistryModule& AssetRegistryModule = FModuleManager::Get().LoadModuleChecked<FAssetRegistryModule>(TEXT("AssetRegistry"));
-
-	if (IsInGameThread())
-	{
-		AssetRegistryModule.Get().WaitForCompletion();
-	}
-
-	return AssetRegistryModule;
-}
+	FAssetRegistryModule& AssetRegistryModule = GetAssetRegistryModule();
 
-void GetTransitiveDependencies(FName PackageName, FAssetRegistryModule& AssetRegistryModule, TSet<FName>& OutDependencies)
-{
 	TArray<FName> PackageQueue;
 	TSet<FName> ExaminedPackages;
 	TArray<FName> PackageDependencies;
 	OutDependencies.Empty();
-	
+
 	PackageQueue.Push(PackageName);
 
 	while (PackageQueue.Num() > 0)
@@ -1354,7 +1730,46 @@ void GetTransitiveDependencies(FName PackageName, FAssetRegistryModule& AssetReg
 	}
 }
 
-void SortPackagesBySize(const IMemoryUsageInfoProvider* MemoryUsageInfoProvider, const TSet<FName>& Packages, TMap<FName, uint64>& OutPackagesWithSize, FOutputDevice* ErrorOutput /* = GLog */)
+static void GetTransitiveDependenciesPackageStore(FName PackageName, TSet<FName>& OutDependencies)
+{
+	const FPackageId PackageId = FPackageId::FromName(PackageName);
+	TSet<FPackageId> TransitiveDependencies;
+	if (!FPackageDependenciesLazyDatabase::Get().GetDependencies(PackageId, TransitiveDependencies))
+	{
+		return;
+	}
+
+	OutDependencies.Empty();
+
+	for (const FPackageId DependencyId : TransitiveDependencies)
+	{
+		FName Name;
+		if (FPackageStoreLazyDatabase::Get().GetPackageNameFromId(DependencyId, Name))
+		{
+			OutDependencies.Add(Name);
+		}
+		else
+		{
+			OutDependencies.Add(FName(EName::Package, DependencyId.Value()));
+		}
+	}
+}
+
+void GetTransitiveDependencies(FName PackageName, TSet<FName>& OutDependencies)
+{
+	if (CVarMemQueryUsePackageStore->GetBool() && FPackageStore::Get().HasAnyBackendsMounted())
+	{
+		GetTransitiveDependenciesPackageStore(PackageName, OutDependencies);
+	}
+	else
+	{
+		GetTransitiveDependenciesAssetRegistry(PackageName, OutDependencies);
+	}
+}
+
+void SortPackagesBySize(
+	const IMemoryUsageInfoProvider* MemoryUsageInfoProvider, const TSet<FName>& Packages, TMap<FName, uint64>& OutPackagesWithSize,
+	FOutputDevice* ErrorOutput /* = GLog */)
 {
 	GetPackagesSize(MemoryUsageInfoProvider, Packages, OutPackagesWithSize, ErrorOutput);
 	OutPackagesWithSize.ValueSort([&](const uint64& A, const uint64& B) { return A > B; });
@@ -1369,7 +1784,7 @@ void GetPackagesSize(const IMemoryUsageInfoProvider* MemoryUsageInfoProvider, co
 	}
 }
 
-void RemoveNonExistentPackages(TMap<FName, uint64>& OutPackagesWithSize)
+static void RemoveNonExistentPackagesAssetRegistry(TMap<FName, uint64>& OutPackagesWithSize)
 {
 	FAssetRegistryModule& AssetRegistryModule = GetAssetRegistryModule();
 
@@ -1382,6 +1797,29 @@ void RemoveNonExistentPackages(TMap<FName, uint64>& OutPackagesWithSize)
 	}
 }
 
+static void RemoveNonExistentPackagesPackageStore(TMap<FName, uint64>& OutPackagesWithSize)
+{
+	for (auto It = OutPackagesWithSize.CreateIterator(); It; ++It)
+	{
+		if (!FPackageStoreLazyDatabase::Get().DoesPackageExists(It.Key()))
+		{
+			It.RemoveCurrent();
+		}
+	}
+}
+
+void RemoveNonExistentPackages(TMap<FName, uint64>& OutPackagesWithSize)
+{
+	if (CVarMemQueryUsePackageStore->GetBool() && FPackageStore::Get().HasAnyBackendsMounted())
+	{
+		RemoveNonExistentPackagesPackageStore(OutPackagesWithSize);
+	}
+	else
+	{
+		RemoveNonExistentPackagesAssetRegistry(OutPackagesWithSize);
+	}
+}
+
 void RemoveFilteredPackages(TMap<FName, uint64>& OutPackagesWithSize, const FString& AssetSubstring)
 {
 	for (auto It = OutPackagesWithSize.CreateIterator(); It; ++It)
diff --git a/Engine/Plugins/MemoryUsageQueries/Source/MemoryUsageQueries/Private/MemoryUsageQueriesPrivate.h b/Engine/Plugins/MemoryUsageQueries/Source/MemoryUsageQueries/Private/MemoryUsageQueriesPrivate.h
--- a/Engine/Plugins/MemoryUsageQueries/Source/MemoryUsageQueries/Private/MemoryUsageQueriesPrivate.h
+++ b/Engine/Plugins/MemoryUsageQueries/Source/MemoryUsageQueries/Private/MemoryUsageQueriesPrivate.h
@@ -3,7 +3,6 @@
 
 #include "UObject/FastReferenceCollector.h"
 
-class FAssetRegistryModule;
 class IMemoryUsageInfoProvider;
 enum class EGCTokenType : uint32;
 
@@ -42,9 +41,9 @@ public:
 
 typedef UE::GC::TDefaultCollector<FMemoryUsageReferenceProcessor> FMemoryUsageReferenceCollector;
 
-bool GetLongName(const FString& ShortPackageName, FName& OutLongPackageName, FOutputDevice* ErrorOutput = GLog);
+bool GetLongName(FStringView ShortPackageName, FName& OutLongPackageName, FOutputDevice* ErrorOutput = GLog);
 bool GetLongNames(const TArray<FString>& PackageNames, TSet<FName>& OutLongPackageNames, FOutputDevice* ErrorOutput = GLog);
-bool GetLongNameAndDependencies(const FString& PackageName, FName& OutLongPackageName, TSet<FName>& OutDependencies, FOutputDevice* ErrorOutput = GLog);
+bool GetLongNameAndDependencies(FStringView PackageName, FName& OutLongPackageName, TSet<FName>& OutDependencies, FOutputDevice* ErrorOutput = GLog);
 
 bool GetDependenciesCombined(const TArray<FString>& PackageNames, TSet<FName>& OutDependencies, FOutputDevice* ErrorOutput = GLog);
 bool GetDependenciesShared(const TArray<FString>& PackageNames, TSet<FName>& OutDependencies, FOutputDevice* ErrorOutput = GLog);
@@ -57,8 +56,7 @@ bool GetRemovablePackages(const TArray<FString>& PackagesToUnload, TSet<FName>&
 // Get Packages that would not be GC'd if PackagesToUnload were unloaded
 bool GetUnremovablePackages(const TArray<FString>& PackagesToUnload, TSet<FName>& OutUnremovablePackages, FOutputDevice* ErrorOutput = GLog);
 
-FAssetRegistryModule& GetAssetRegistryModule();
-void GetTransitiveDependencies(FName PackageName, FAssetRegistryModule& AssetRegistryModule, TSet<FName>& OutDependencies);
+void GetTransitiveDependencies(FName PackageName, TSet<FName>& OutDependencies);
 
 void SortPackagesBySize(const IMemoryUsageInfoProvider* MemoryUsageInfoProvider, const TSet<FName>& Packages, TMap<FName, uint64>& OutPackagesWithSize, FOutputDevice* ErrorOutput = GLog);
 void GetPackagesSize(const IMemoryUsageInfoProvider* MemoryUsageInfoProvider, const TSet<FName>& Packages, TMap<FName, uint64>& OutPackagesWithSize, FOutputDevice* ErrorOutput = GLog);
diff --git a/Engine/Source/Runtime/Core/Private/IO/PackageStore.cpp b/Engine/Source/Runtime/Core/Private/IO/PackageStore.cpp
--- a/Engine/Source/Runtime/Core/Private/IO/PackageStore.cpp
+++ b/Engine/Source/Runtime/Core/Private/IO/PackageStore.cpp
@@ -208,3 +208,7 @@ FPackageStoreBackendContext::FPendingEntriesAddedEvent& FPackageStore::OnPending
 	return BackendContext->PendingEntriesAdded;
 }
 
+bool FPackageStore::HasAnyBackendsMounted() const
+{
+	return !Backends.IsEmpty();
+}
diff --git a/Engine/Source/Runtime/Core/Public/IO/PackageStore.h b/Engine/Source/Runtime/Core/Public/IO/PackageStore.h
--- a/Engine/Source/Runtime/Core/Public/IO/PackageStore.h
+++ b/Engine/Source/Runtime/Core/Public/IO/PackageStore.h
@@ -235,6 +235,8 @@ public:
 
 	CORE_API FPackageStoreBackendContext::FPendingEntriesAddedEvent& OnPendingEntriesAdded();
 
+	CORE_API bool HasAnyBackendsMounted() const;
+
 private:
 	FPackageStore();
 
diff --git a/Engine/Source/Runtime/PakFile/Private/IPlatformFilePak.cpp b/Engine/Source/Runtime/PakFile/Private/IPlatformFilePak.cpp
--- a/Engine/Source/Runtime/PakFile/Private/IPlatformFilePak.cpp
+++ b/Engine/Source/Runtime/PakFile/Private/IPlatformFilePak.cpp
@@ -324,6 +324,49 @@ void FPakPlatformFile::GetFilenamesFromIostoreContainer(const FString& InContain
 	}
 }
 
+void FPakPlatformFile::ForeachPackageInIostoreWhile(TFunctionRef<bool(FName)> Predicate)
+{
+	FPakPlatformFile* PakPlatformFile = static_cast<FPakPlatformFile*>(FPlatformFileManager::Get().FindPlatformFile(FPakPlatformFile::GetTypeName()));
+	if (!PakPlatformFile || !PakPlatformFile->IoDispatcherFileBackend.IsValid())
+	{
+		return;
+	}
+	FScopeLock ScopedLock(&PakPlatformFile->PakListCritical);
+	for (const FPakListEntry& PakListEntry : PakPlatformFile->PakFiles)
+	{
+		TUniquePtr<FIoStoreReader> IoStoreReader(new FIoStoreReader());
+		FIoStatus Status = IoStoreReader->Initialize(*FPaths::ChangeExtension(PakListEntry.PakFile->PakFilename, TEXT("")), GetRegisteredEncryptionKeys().GetKeys());
+		if (Status.IsOk())
+		{
+			const FIoDirectoryIndexReader& DirectoryIndex = IoStoreReader->GetDirectoryIndexReader();
+
+			const bool Result = DirectoryIndex.IterateDirectoryIndex(
+				FIoDirectoryIndexHandle::RootDirectory(),
+				TEXT(""),
+				[Predicate](FString Filename, uint32) -> bool
+				{
+					const FString Ext = FPaths::GetExtension(Filename);
+					if (Ext != TEXT("umap") && Ext != TEXT("uasset"))
+					{
+						return true; // ignore non package files
+					}
+
+					FString PackageName;
+					if (FPackageName::TryConvertFilenameToLongPackageName(Filename, PackageName))
+					{
+						return Invoke(Predicate, FName(*PackageName));
+					}
+
+					return true; // ignore not mapped packages
+				});
+			if (!Result)
+			{
+				return;
+			}
+		}
+	}
+}
+
 #if !defined(PLATFORM_BYPASS_PAK_PRECACHE)
 	#error "PLATFORM_BYPASS_PAK_PRECACHE must be defined."
 #endif
diff --git a/Engine/Source/Runtime/PakFile/Public/IPlatformFilePak.h b/Engine/Source/Runtime/PakFile/Public/IPlatformFilePak.h
--- a/Engine/Source/Runtime/PakFile/Public/IPlatformFilePak.h
+++ b/Engine/Source/Runtime/PakFile/Public/IPlatformFilePak.h
@@ -3114,6 +3114,9 @@ public:
 
 	/** Returns the RelativePathFromMount Filename for every Filename found in the Iostore Container that relates to the provided block indexes */
 	static void GetFilenamesFromIostoreByBlockIndex(const FString& InContainerName, const TArray<int32>& InBlockIndex, TArray<FString>& OutFileList);
+	
+	/** Iterates Iostore Container while Predicate returns true */
+	static void ForeachPackageInIostoreWhile(TFunctionRef<bool(FName)> Predicate);
 
 	void ReleaseOldReaders();
 
