Fix CIS compile error: Missing PakFile dependency

#rnx

diff --git a/Engine/Plugins/MemoryUsageQueries/Source/MemoryUsageQueries/MemoryUsageQueries.Build.cs b/Engine/Plugins/MemoryUsageQueries/Source/MemoryUsageQueries/MemoryUsageQueries.Build.cs
--- a/Engine/Plugins/MemoryUsageQueries/Source/MemoryUsageQueries/MemoryUsageQueries.Build.cs
+++ b/Engine/Plugins/MemoryUsageQueries/Source/MemoryUsageQueries/MemoryUsageQueries.Build.cs
@@ -28,7 +28,8 @@ namespace UnrealBuildTool.Rules
 				{
 					"AssetRegistry",
 					"Engine",
-					"EngineSettings"
+					"EngineSettings",
+					"PakFile"
 				}
 			);
 		}
