[MemQuery] Fix walking dependency tree by consuming same element instead of peeking the last element and removing the first.

#tests replay
#rnx

diff --git a/Engine/Plugins/MemoryUsageQueries/Source/MemoryUsageQueries/Private/MemoryUsageQueries.cpp b/Engine/Plugins/MemoryUsageQueries/Source/MemoryUsageQueries/Private/MemoryUsageQueries.cpp
--- a/Engine/Plugins/MemoryUsageQueries/Source/MemoryUsageQueries/Private/MemoryUsageQueries.cpp
+++ b/Engine/Plugins/MemoryUsageQueries/Source/MemoryUsageQueries/Private/MemoryUsageQueries.cpp
@@ -1364,15 +1364,15 @@ class FPackageDependenciesLazyDatabase final
 	{
 		LLM_SCOPE(TEXT("MemoryUsageQueries"));
 
-		TArray<FPackageId, TInlineAllocator<2048>> Queue;
+		TArray<FPackageId, TInlineAllocator<2048>> Stack;
 
-		Queue.Add(RootPackageId);
+		Stack.Add(RootPackageId);
 		bool bAddedSuccessfully = false;
 
-		while (!Queue.IsEmpty())
+		const bool bShouldShrink = false;
+		while (!Stack.IsEmpty())
 		{
-			const FPackageId PackageId = Queue.Top();
-			Queue.RemoveAt(0, 1, false);
+			const FPackageId PackageId = Stack.Pop(bShouldShrink);
 
 			if (Dependencies.Contains(PackageId) || Leafs.Contains(PackageId))
 			{
@@ -1392,7 +1392,7 @@ class FPackageDependenciesLazyDatabase final
 					Referencers.FindOrAdd(DependentId).Add(PackageId);
 					bAddedSuccessfully = true;
 				}
-				Queue.Append(ImportedPackageIds);
+				Stack.Append(ImportedPackageIds);
 
 #if WITH_EDITOR
 				// Add editor optional dependencies
@@ -1403,7 +1403,7 @@ class FPackageDependenciesLazyDatabase final
 					Referencers.FindOrAdd(DependentId).Add(PackageId);
 					bAddedSuccessfully = true;
 				}
-				Queue.Append(OptionalImportedPackageIds);
+				Stack.Append(OptionalImportedPackageIds);
 #endif
 
 
@@ -1458,17 +1458,18 @@ public:
 			return false;
 		}
 
-		TArray<FPackageId, TInlineAllocator<2048>> Children(ChildrenSet->Array());
-		while (!Children.IsEmpty())
+		const bool bShouldShrink = false;
+		TArray<FPackageId, TInlineAllocator<2048>> Stack(ChildrenSet->Array());
+		while (!Stack.IsEmpty())
 		{
-			FPackageId Child = Children.Top();
-			Children.RemoveAt(0, 1, false);
+			FPackageId Child = Stack.Pop(bShouldShrink);
+
 			if (!OutDependencies.Contains(Child))
 			{
 				OutDependencies.Add(Child);
 				if (const TSet<FPackageId>* GrandChildrenSet = Dependencies.Find(Child))
 				{
-					Children.Append(GrandChildrenSet->Array());
+					Stack.Append(GrandChildrenSet->Array());
 				}
 			}
 		}
