Allow the writing of files to the a MemQuery folder if LOG=<filename> is specified as part of the command.

diff --git a/Engine/Plugins/MemoryUsageQueries/Source/MemoryUsageQueries/Private/MemoryUsageQueries.cpp b/Engine/Plugins/MemoryUsageQueries/Source/MemoryUsageQueries/Private/MemoryUsageQueries.cpp
--- a/Engine/Plugins/MemoryUsageQueries/Source/MemoryUsageQueries/Private/MemoryUsageQueries.cpp
+++ b/Engine/Plugins/MemoryUsageQueries/Source/MemoryUsageQueries/Private/MemoryUsageQueries.cpp
@@ -11,6 +11,7 @@
 #include "AssetRegistry/AssetRegistryModule.h"
 #include "Engine/BlueprintGeneratedClass.h"
 #include "IO/PackageStore.h"
+#include "Misc/OutputDeviceArchiveWrapper.h"
 #include "Misc/PackageName.h"
 #include "Misc/WildcardString.h"
 #include "Modules/ModuleManager.h"
@@ -78,8 +79,15 @@ static int32 DefaultResultLimit = 15;
 
 bool FMemoryUsageQueriesExec::Exec_Runtime(UWorld* InWorld, const TCHAR* Cmd, FOutputDevice& Ar)
 {
+	bool bResult = false;
+
 	if (FParse::Command(&Cmd, TEXT("MemQuery")))
 	{
+		FOutputDevice* OutputDevice = &Ar;
+
+		FOutputDeviceArchiveWrapper* FileArWrapper = nullptr;
+		FArchive* FileAr = nullptr;
+
 		Ar.Logf(TEXT("MemQuery: %s"), Cmd);
 
 		const bool bTruncate = !FParse::Param(Cmd, TEXT("notrunc"));
@@ -95,6 +103,33 @@ bool FMemoryUsageQueriesExec::Exec_Runtime(UWorld* InWorld, const TCHAR* Cmd, FO
 		int32 Limit = -1;
 		FParse::Value(Cmd, TEXT("Limit="), Limit);
 
+		FString LogFileName;
+		FParse::Value(Cmd, TEXT("Log="), LogFileName);
+
+#if ALLOW_DEBUG_FILES
+		if (!LogFileName.IsEmpty())
+		{
+			// Create folder for MemQuery files.
+			const FString OutputDir = FPaths::ProfilingDir() / TEXT("MemQuery");
+			IFileManager::Get().MakeDirectory(*OutputDir, true);
+
+			FString FileTimeString = FString::Printf(TEXT("_%s"), *FDateTime::Now().ToString(TEXT("%H%M%S")));
+			const FString FileExtension = (bCSV ? TEXT(".csv") : TEXT(".memquery"));
+			const FString LogFilename = OutputDir / (LogFileName + FileTimeString + FileExtension);
+
+			FileAr = IFileManager::Get().CreateDebugFileWriter(*LogFilename);
+			if (FileAr != nullptr)
+			{
+				FileArWrapper = new FOutputDeviceArchiveWrapper(FileAr);				
+			}
+
+			if (FileArWrapper != nullptr)
+			{
+				OutputDevice = FileArWrapper;
+			}
+		}
+#endif
+
 		if (FParse::Command(&Cmd, TEXT("Usage")))
 		{
 			if (!Name.IsEmpty())
@@ -104,14 +139,14 @@ bool FMemoryUsageQueriesExec::Exec_Runtime(UWorld* InWorld, const TCHAR* Cmd, FO
 
 				if (MemoryUsageQueries::GetMemoryUsage(MemoryUsageQueries::CurrentMemoryUsageInfoProvider, Name, ExclusiveSize, InclusiveSize, &Ar))
 				{
-					Ar.Logf(TEXT("MemoryUsage: ExclusiveSize: %.2f MiB (%.2f KiB); InclusiveSize: %.2f MiB (%.2f KiB)"),
+					OutputDevice->Logf(TEXT("MemoryUsage: ExclusiveSize: %.2f MiB (%.2f KiB); InclusiveSize: %.2f MiB (%.2f KiB)"),
 							ExclusiveSize / (1024.f * 1024.f),
 							ExclusiveSize / 1024.f,
 							InclusiveSize / (1024.f * 1024.f),
 							InclusiveSize / 1024.f);
 				}
 
-				return true;
+				bResult = true;
 			}
 		}
 		else if (FParse::Command(&Cmd, TEXT("CombinedUsage")))
@@ -124,10 +159,10 @@ bool FMemoryUsageQueriesExec::Exec_Runtime(UWorld* InWorld, const TCHAR* Cmd, FO
 
 				if (MemoryUsageQueries::GetMemoryUsageCombined(MemoryUsageQueries::CurrentMemoryUsageInfoProvider, Packages, TotalSize, &Ar))
 				{
-					Ar.Logf(TEXT("MemoryUsageCombined: TotalSize: %.2f MiB (%.2f KiB)"), TotalSize / (1024.f * 1024.f), TotalSize / 1024.f);
+					OutputDevice->Logf(TEXT("MemoryUsageCombined: TotalSize: %.2f MiB (%.2f KiB)"), TotalSize / (1024.f * 1024.f), TotalSize / 1024.f);
 				}
 
-				return true;
+				bResult = true;
 			}
 		}
 		else if (FParse::Command(&Cmd, TEXT("SharedUsage")))
@@ -140,10 +175,10 @@ bool FMemoryUsageQueriesExec::Exec_Runtime(UWorld* InWorld, const TCHAR* Cmd, FO
 
 				if (MemoryUsageQueries::GetMemoryUsageShared(MemoryUsageQueries::CurrentMemoryUsageInfoProvider, Packages, SharedSize, &Ar))
 				{
-					Ar.Logf(TEXT("MemoryUsageShared: SharedSize: %.2f MiB (%.2f KiB)"), SharedSize / (1024.f * 1024.f), SharedSize / 1024.f);
+					OutputDevice->Logf(TEXT("MemoryUsageShared: SharedSize: %.2f MiB (%.2f KiB)"), SharedSize / (1024.f * 1024.f), SharedSize / 1024.f);
 				}
 
-				return true;
+				bResult = true;
 			}
 		}
 		else if (FParse::Command(&Cmd, TEXT("UniqueUsage")))
@@ -156,10 +191,10 @@ bool FMemoryUsageQueriesExec::Exec_Runtime(UWorld* InWorld, const TCHAR* Cmd, FO
 
 				if (MemoryUsageQueries::GetMemoryUsageUnique(MemoryUsageQueries::CurrentMemoryUsageInfoProvider, Packages, UniqueSize, &Ar))
 				{
-					Ar.Logf(TEXT("MemoryUsageUnique: UniqueSize: %.2f MiB (%.2f KiB)"), UniqueSize / (1024.f * 1024.f), UniqueSize / 1024.f);
+					OutputDevice->Logf(TEXT("MemoryUsageUnique: UniqueSize: %.2f MiB (%.2f KiB)"), UniqueSize / (1024.f * 1024.f), UniqueSize / 1024.f);
 				}
 
-				return true;
+				bResult = true;
 			}
 		}
 		else if (FParse::Command(&Cmd, TEXT("CommonUsage")))
@@ -172,10 +207,10 @@ bool FMemoryUsageQueriesExec::Exec_Runtime(UWorld* InWorld, const TCHAR* Cmd, FO
 
 				if (MemoryUsageQueries::GetMemoryUsageCommon(MemoryUsageQueries::CurrentMemoryUsageInfoProvider, Packages, CommonSize, &Ar))
 				{
-					Ar.Logf(TEXT("MemoryUsageCommon: CommonSize: %.2f MiB (%.2f KiB)"), CommonSize / (1024.f * 1024.f), CommonSize / 1024.f);
+					OutputDevice->Logf(TEXT("MemoryUsageCommon: CommonSize: %.2f MiB (%.2f KiB)"), CommonSize / (1024.f * 1024.f), CommonSize / 1024.f);
 				}
 
-				return true;
+				bResult = true;
 			}
 		}
 		else if (FParse::Command(&Cmd, TEXT("Dependencies")))
@@ -186,10 +221,10 @@ bool FMemoryUsageQueriesExec::Exec_Runtime(UWorld* InWorld, const TCHAR* Cmd, FO
 
 				if (MemoryUsageQueries::GetDependenciesWithSize(MemoryUsageQueries::CurrentMemoryUsageInfoProvider, Name, DependenciesWithSize, &Ar))
 				{
-					MemoryUsageQueries::Internal::PrintTagsWithSize(Ar, DependenciesWithSize, TEXT("dependencies"), bTruncate, Limit, bCSV);
+					MemoryUsageQueries::Internal::PrintTagsWithSize(*OutputDevice, DependenciesWithSize, TEXT("dependencies"), bTruncate, Limit, bCSV);
 				}
 
-				return true;
+				bResult = true;
 			}
 		}
 		else if (FParse::Command(&Cmd, TEXT("CombinedDependencies")))
@@ -202,10 +237,10 @@ bool FMemoryUsageQueriesExec::Exec_Runtime(UWorld* InWorld, const TCHAR* Cmd, FO
 
 				if (MemoryUsageQueries::GetDependenciesWithSizeCombined(MemoryUsageQueries::CurrentMemoryUsageInfoProvider, Packages, CombinedDependenciesWithSize, &Ar))
 				{
-					MemoryUsageQueries::Internal::PrintTagsWithSize(Ar, CombinedDependenciesWithSize, TEXT("combined dependencies"), bTruncate, Limit, bCSV);
+					MemoryUsageQueries::Internal::PrintTagsWithSize(*OutputDevice, CombinedDependenciesWithSize, TEXT("combined dependencies"), bTruncate, Limit, bCSV);
 				}
 
-				return true;
+				bResult = true;
 			}
 		}
 		else if (FParse::Command(&Cmd, TEXT("SharedDependencies")))
@@ -218,10 +253,10 @@ bool FMemoryUsageQueriesExec::Exec_Runtime(UWorld* InWorld, const TCHAR* Cmd, FO
 
 				if (MemoryUsageQueries::GetDependenciesWithSizeShared(MemoryUsageQueries::CurrentMemoryUsageInfoProvider, Packages, SharedDependenciesWithSize, &Ar))
 				{
-					MemoryUsageQueries::Internal::PrintTagsWithSize(Ar, SharedDependenciesWithSize, TEXT("shared dependencies"), bTruncate, Limit, bCSV);
+					MemoryUsageQueries::Internal::PrintTagsWithSize(*OutputDevice, SharedDependenciesWithSize, TEXT("shared dependencies"), bTruncate, Limit, bCSV);
 				}
 
-				return true;
+				bResult = true;
 			}
 		}
 		else if (FParse::Command(&Cmd, TEXT("UniqueDependencies")))
@@ -234,10 +269,10 @@ bool FMemoryUsageQueriesExec::Exec_Runtime(UWorld* InWorld, const TCHAR* Cmd, FO
 
 				if (MemoryUsageQueries::GetDependenciesWithSizeUnique(MemoryUsageQueries::CurrentMemoryUsageInfoProvider, Packages, UniqueDependenciesWithSize, &Ar))
 				{
-					MemoryUsageQueries::Internal::PrintTagsWithSize(Ar, UniqueDependenciesWithSize, TEXT("unique dependencies"), bTruncate, Limit, bCSV);
+					MemoryUsageQueries::Internal::PrintTagsWithSize(*OutputDevice, UniqueDependenciesWithSize, TEXT("unique dependencies"), bTruncate, Limit, bCSV);
 				}
 
-				return true;
+				bResult = true;
 			}
 		}
 		else if (FParse::Command(&Cmd, TEXT("CommonDependencies")))
@@ -250,10 +285,10 @@ bool FMemoryUsageQueriesExec::Exec_Runtime(UWorld* InWorld, const TCHAR* Cmd, FO
 
 				if (MemoryUsageQueries::GetDependenciesWithSizeCommon(MemoryUsageQueries::CurrentMemoryUsageInfoProvider, Packages, CommonDependenciesWithSize, &Ar))
 				{
-					MemoryUsageQueries::Internal::PrintTagsWithSize(Ar, CommonDependenciesWithSize, TEXT("common dependencies"), bTruncate, Limit, bCSV);
+					MemoryUsageQueries::Internal::PrintTagsWithSize(*OutputDevice, CommonDependenciesWithSize, TEXT("common dependencies"), bTruncate, Limit, bCSV);
 				}
 
-				return true;
+				bResult = true;
 			}
 		}
 #if ENABLE_LOW_LEVEL_MEM_TRACKER
@@ -291,10 +326,10 @@ bool FMemoryUsageQueriesExec::Exec_Runtime(UWorld* InWorld, const TCHAR* Cmd, FO
 
 			if (bSuccess)
 			{
-				MemoryUsageQueries::Internal::PrintTagsWithSize(Ar, AssetsWithSize, TEXT("largest assets"), bTruncate, Limit, bCSV);
+				MemoryUsageQueries::Internal::PrintTagsWithSize(*OutputDevice, AssetsWithSize, TEXT("largest assets"), bTruncate, Limit, bCSV);
 			}
 
-			return true;
+			bResult = true;
 		}
 		else if (FParse::Command(&Cmd, TEXT("ListClasses")))
 		{
@@ -312,10 +347,10 @@ bool FMemoryUsageQueriesExec::Exec_Runtime(UWorld* InWorld, const TCHAR* Cmd, FO
 
 			if (MemoryUsageQueries::GetFilteredClassesWithSize(ClassesWithSize, Group, AssetName, &Ar))
 			{
-				MemoryUsageQueries::Internal::PrintTagsWithSize(Ar, ClassesWithSize, *FString::Printf(TEXT("largest classes")), bTruncate, Limit, bCSV);
+				MemoryUsageQueries::Internal::PrintTagsWithSize(*OutputDevice, ClassesWithSize, *FString::Printf(TEXT("largest classes")), bTruncate, Limit, bCSV);
 			}
 
-			return true;
+			bResult = true;
 		}
 		else if (FParse::Command(&Cmd, TEXT("ListGroups")))
 		{
@@ -333,10 +368,10 @@ bool FMemoryUsageQueriesExec::Exec_Runtime(UWorld* InWorld, const TCHAR* Cmd, FO
 
 			if (MemoryUsageQueries::GetFilteredGroupsWithSize(GroupsWithSize, AssetName, Class, &Ar))
 			{
-				MemoryUsageQueries::Internal::PrintTagsWithSize(Ar, GroupsWithSize, *FString::Printf(TEXT("largest groups")), bTruncate, Limit, bCSV);
+				MemoryUsageQueries::Internal::PrintTagsWithSize(*OutputDevice, GroupsWithSize, *FString::Printf(TEXT("largest groups")), bTruncate, Limit, bCSV);
 			}
 
-			return true;
+			bResult = true;
 		}
 #endif
 		else if (FParse::Command(&Cmd, TEXT("Savings")))
@@ -384,9 +419,9 @@ bool FMemoryUsageQueriesExec::Exec_Runtime(UWorld* InWorld, const TCHAR* Cmd, FO
 				}
 
 				PresetSavings.ValueSort(TGreater<uint64>());
-				MemoryUsageQueries::Internal::PrintTagsWithSize(Ar, PresetSavings, *FString::Printf(TEXT("possible savings")), bTruncate, bCSV);
+				MemoryUsageQueries::Internal::PrintTagsWithSize(*OutputDevice, PresetSavings, *FString::Printf(TEXT("possible savings")), bTruncate, bCSV);
 
-				return true;
+				bResult = true;
 			}
 		}
 #if ENABLE_LOW_LEVEL_MEM_TRACKER
@@ -411,7 +446,7 @@ bool FMemoryUsageQueriesExec::Exec_Runtime(UWorld* InWorld, const TCHAR* Cmd, FO
 				if (!bSuccess)
 				{
 					Ar.Logf(TEXT("Failed to gather assets for Collection %s"), *CollectionInfo.Name);
-					return false;
+					break;
 				}
 
 				// Will return true if the Package Name matches any of the conditions in the array of Paths
@@ -462,7 +497,7 @@ bool FMemoryUsageQueriesExec::Exec_Runtime(UWorld* InWorld, const TCHAR* Cmd, FO
 				if (!MemoryUsageQueries::GatherDependenciesForPackages(MemoryUsageQueries::CurrentMemoryUsageInfoProvider, PackageNames, InternalDependencies, ExternalDependencies, &Ar))
 				{
 					Ar.Logf(TEXT("Failed to gather memory usage for dependencies in Collection %s"), *CollectionInfo.Name);
-					return false;
+					break;
 				}
 
 				uint64 TotalCollectionSize = 0;
@@ -519,11 +554,11 @@ bool FMemoryUsageQueriesExec::Exec_Runtime(UWorld* InWorld, const TCHAR* Cmd, FO
 
 				if (bCSV)
 				{
-					Ar.Logf(TEXT(",Asset,Exclusive KiB,Unique Refs KiB,Unique Ref Count,Shared Refs KiB,Shared Ref Count,Total KiB"));
+					OutputDevice->Logf(TEXT(",Asset,Exclusive KiB,Unique Refs KiB,Unique Ref Count,Shared Refs KiB,Shared Ref Count,Total KiB"));
 				}
 				else
 				{
-					Ar.Logf(
+					OutputDevice->Logf(
 						TEXT(" %100s %20s %20s %15s %20s %15s %25s"),
 						TEXT("Asset"),
 						TEXT("Exclusive KiB"),
@@ -543,7 +578,7 @@ bool FMemoryUsageQueriesExec::Exec_Runtime(UWorld* InWorld, const TCHAR* Cmd, FO
 
 					if (bCSV)
 					{
-						Ar.Logf(TEXT(",%s,%.2f,%.2f,%d,%.2f,%d,%.2f"),*Asset.Key.ToString(),
+						OutputDevice->Logf(TEXT(",%s,%.2f,%.2f,%d,%.2f,%d,%.2f"),*Asset.Key.ToString(),
 							AssetMemoryDetails.ExclusiveSize / 1024.f, 
 							AssetMemoryDetails.UniqueSize / 1024.f, 
 							AssetMemory.UniqueRefCount, 
@@ -553,7 +588,7 @@ bool FMemoryUsageQueriesExec::Exec_Runtime(UWorld* InWorld, const TCHAR* Cmd, FO
 					}
 					else
 					{
-						Ar.Logf(
+						OutputDevice->Logf(
 							TEXT(" %100s %20.2f %20.2f %15d %20.2f %15d %25.2f"),
 							*Asset.Key.ToString(),
 							AssetMemoryDetails.ExclusiveSize / 1024.f, 
@@ -571,11 +606,11 @@ bool FMemoryUsageQueriesExec::Exec_Runtime(UWorld* InWorld, const TCHAR* Cmd, FO
 				{
 					if (bCSV)
 					{
-						Ar.Logf(TEXT(",Asset,Dependency,Unique KiB,Shared KiB"));
+						OutputDevice->Logf(TEXT(",Asset,Dependency,Unique KiB,Shared KiB"));
 					}
 					else
 					{
-						Ar.Logf(TEXT(" %100s %100s %20s %20s"),
+						OutputDevice->Logf(TEXT(" %100s %100s %20s %20s"),
 							TEXT("Asset"),
 							TEXT("Dependency"),
 							TEXT("Unique KiB"),
@@ -592,12 +627,12 @@ bool FMemoryUsageQueriesExec::Exec_Runtime(UWorld* InWorld, const TCHAR* Cmd, FO
 
 							if (bCSV)
 							{
-								Ar.Logf(TEXT(",%s,%s,%.2f,%.2f"), *Asset.Key.ToString(), *DependencyAssetName,
+								OutputDevice->Logf(TEXT(",%s,%s,%.2f,%.2f"), *Asset.Key.ToString(), *DependencyAssetName,
 									DepedencyMemoryDetails.UniqueSize / 1024.f, DepedencyMemoryDetails.SharedSize / 1024.f);
 							}
 							else
 							{
-								Ar.Logf(TEXT(" %100s %100s %20.2f %20.2f"), 
+								OutputDevice->Logf(TEXT(" %100s %100s %20.2f %20.2f"),
 									*Asset.Key.ToString(), 
 									*DependencyAssetName,
 									DepedencyMemoryDetails.UniqueSize / 1024.f, 
@@ -610,20 +645,28 @@ bool FMemoryUsageQueriesExec::Exec_Runtime(UWorld* InWorld, const TCHAR* Cmd, FO
 
 				if (bCSV)
 				{
-					Ar.Logf(TEXT(",TOTAL KiB,%.2f"), TotalCollectionSize / 1024.f);
+					OutputDevice->Logf(TEXT(",TOTAL KiB,%.2f"), TotalCollectionSize / 1024.f);
 				}
 				else
 				{
-					Ar.Logf(TEXT("TOTAL KiB: %.2f"), TotalCollectionSize / 1024.f);
+					OutputDevice->Logf(TEXT("TOTAL KiB: %.2f"), TotalCollectionSize / 1024.f);
 				}
 
-				return true;
+				bResult = true;
 			}
 		}
 #endif
-	}
 
-	return false;
+		if (FileArWrapper != nullptr)
+		{
+			FileArWrapper->TearDown();
+		}
+			
+		delete FileArWrapper;
+		delete FileAr;
+	}
+	
+	return bResult;
 }
 
 namespace MemoryUsageQueries
