Converted MemQuery to FAutoConsoleCommand from FSelfRegisteringExec
All commands now support logging to a file

diff --git a/Engine/Plugins/MemoryUsageQueries/Source/MemoryUsageQueries/Private/MemoryUsageQueries.cpp b/Engine/Plugins/MemoryUsageQueries/Source/MemoryUsageQueries/Private/MemoryUsageQueries.cpp
--- a/Engine/Plugins/MemoryUsageQueries/Source/MemoryUsageQueries/Private/MemoryUsageQueries.cpp
+++ b/Engine/Plugins/MemoryUsageQueries/Source/MemoryUsageQueries/Private/MemoryUsageQueries.cpp
@@ -28,36 +28,21 @@ IMemoryUsageInfoProvider* CurrentMemoryUsageInfoProvider = &MemoryUsageInfoProvi
 
 } // namespace MemoryUsageQueries
 
-struct FMemoryUsageQueriesExec : public FSelfRegisteringExec
-{
-	FMemoryUsageQueriesExec() {}
-
-protected:
-	virtual bool Exec_Runtime(UWorld* Inworld, const TCHAR* Cmd, FOutputDevice& Ar) override;
-};
-
 struct FAssetMemoryBreakdown
 {
-	FAssetMemoryBreakdown() :
-		ExclusiveSize(0),
-		UniqueSize(0),
-		SharedSize(0),
-		TotalSize(0)
+	FAssetMemoryBreakdown()
 	{
-
 	}
 
-	uint64 ExclusiveSize;
-	uint64 UniqueSize;
-	uint64 SharedSize;
-	uint64 TotalSize;
+	uint64 ExclusiveSize = 0;
+	uint64 UniqueSize = 0;
+	uint64 SharedSize = 0;
+	uint64 TotalSize = 0;
 };
 
 struct FAssetMemoryDetails
 {
-	FAssetMemoryDetails() :
-		UniqueRefCount(0),
-		SharedRefCount(0)
+	FAssetMemoryDetails()
 	{
 	}
 
@@ -71,736 +56,773 @@ struct FAssetMemoryDetails
 
 	TMap<FName, FAssetMemoryBreakdown> DependenciesToMemoryMap;
 
-	int32 UniqueRefCount;
-	int32 SharedRefCount;
+	int32 UniqueRefCount = 0;
+	int32 SharedRefCount = 0;
 };
 
-static FMemoryUsageQueriesExec DebugSettinsSubsystemExecInstance;
 static int32 DefaultResultLimit = 15;
 
-bool FMemoryUsageQueriesExec::Exec_Runtime(UWorld* InWorld, const TCHAR* Cmd, FOutputDevice& Ar)
+// Helper struct that scopes an output device, and will automatically clean up any log files when leaving scope
+struct FScopedOutputDevice
 {
-	bool bResult = false;
+	FOutputDevice* CurrentOutputDevice = nullptr;
+
+	FOutputDeviceArchiveWrapper* FileArWrapper = nullptr;
+	FArchive* FileAr = nullptr;
 
-	if (FParse::Command(&Cmd, TEXT("MemQuery")))
+	explicit FScopedOutputDevice(FOutputDevice* InDefaultOutputDevice) :
+		CurrentOutputDevice(InDefaultOutputDevice)
+		
 	{
-		FOutputDevice* OutputDevice = &Ar;
+	}
 
-		FOutputDeviceArchiveWrapper* FileArWrapper = nullptr;
-		FArchive* FileAr = nullptr;
+	~FScopedOutputDevice()
+	{
+		if (FileArWrapper != nullptr)
+		{
+			FileArWrapper->TearDown();
+		}
 
-		Ar.Logf(TEXT("MemQuery: %s"), Cmd);
+		delete FileArWrapper;
+		delete FileAr;
+	}
 
-		const bool bTruncate = !FParse::Param(Cmd, TEXT("notrunc"));
-		const bool bCSV = FParse::Param(Cmd, TEXT("csv"));
+	FOutputDevice& GetOutputDevice() const
+	{
+		return *CurrentOutputDevice;
+	}
 
-		// Parse some common options
-		FString Name;
-		FParse::Value(Cmd, TEXT("Name="), Name);
+	void OpenLogFile(const FString& LogFileName, bool bCSV)
+	{
+		if (FileArWrapper != nullptr || LogFileName.IsEmpty())
+		{
+			// Log file is already open or LogFileName is empty
+			return;
+		}
 
-		FString Names;
-		FParse::Value(Cmd, TEXT("Names="), Names);
+#if ALLOW_DEBUG_FILES
+		// Create folder for MemQuery files.
+		const FString OutputDir = FPaths::ProfilingDir() / TEXT("MemQuery");
+		IFileManager::Get().MakeDirectory(*OutputDir, true);
 
-		int32 Limit = -1;
-		FParse::Value(Cmd, TEXT("Limit="), Limit);
+		FString FileTimeString = FString::Printf(TEXT("_%s"), *FDateTime::Now().ToString(TEXT("%H%M%S")));
+		const FString FileExtension = (bCSV ? TEXT(".csv") : TEXT(".memquery"));
+		const FString LogFilename = OutputDir / (LogFileName + FileTimeString + FileExtension);
 
-		FString LogFileName;
-		FParse::Value(Cmd, TEXT("Log="), LogFileName);
+		FileAr = IFileManager::Get().CreateDebugFileWriter(*LogFilename);
 
-#if ALLOW_DEBUG_FILES
-		if (!LogFileName.IsEmpty())
+		if (FileAr != nullptr)
 		{
-			// Create folder for MemQuery files.
-			const FString OutputDir = FPaths::ProfilingDir() / TEXT("MemQuery");
-			IFileManager::Get().MakeDirectory(*OutputDir, true);
-
-			FString FileTimeString = FString::Printf(TEXT("_%s"), *FDateTime::Now().ToString(TEXT("%H%M%S")));
-			const FString FileExtension = (bCSV ? TEXT(".csv") : TEXT(".memquery"));
-			const FString LogFilename = OutputDir / (LogFileName + FileTimeString + FileExtension);
-
-			FileAr = IFileManager::Get().CreateDebugFileWriter(*LogFilename);
-			if (FileAr != nullptr)
+			FileArWrapper = new FOutputDeviceArchiveWrapper(FileAr);
+				
+			if (FileArWrapper != nullptr)
 			{
-				FileArWrapper = new FOutputDeviceArchiveWrapper(FileAr);				
+				CurrentOutputDevice = FileArWrapper;
 			}
-
-			if (FileArWrapper != nullptr)
+			else
 			{
-				OutputDevice = FileArWrapper;
+				delete FileAr;
+				FileAr = nullptr;
 			}
 		}
 #endif
+	}
 
-		if (FParse::Command(&Cmd, TEXT("Usage")))
-		{
-			if (!Name.IsEmpty())
-			{
-				uint64 ExclusiveSize = 0U;
-				uint64 InclusiveSize = 0U;
+	// Delete default constructor to force proper usage
+	FScopedOutputDevice() = delete;
+};
 
-				if (MemoryUsageQueries::GetMemoryUsage(MemoryUsageQueries::CurrentMemoryUsageInfoProvider, Name, ExclusiveSize, InclusiveSize, &Ar))
-				{
-					OutputDevice->Logf(TEXT("MemoryUsage: ExclusiveSize: %.2f MiB (%.2f KiB); InclusiveSize: %.2f MiB (%.2f KiB)"),
-							ExclusiveSize / (1024.f * 1024.f),
-							ExclusiveSize / 1024.f,
-							InclusiveSize / (1024.f * 1024.f),
-							InclusiveSize / 1024.f);
-				}
+/** Structure that will parse a string and fill out some commonly used parameters per command */
+struct FCommonParameters
+{
+	int32 Limit = -1;
 
-				bResult = true;
-			}
-		}
-		else if (FParse::Command(&Cmd, TEXT("CombinedUsage")))
-		{
-			if (!Names.IsEmpty())
-			{
-				TArray<FString> Packages;
-				Names.ParseIntoArrayWS(Packages);
-				uint64 TotalSize;
+	bool bTruncate = true;
+	bool bCSV = false;
+	
+	FName Group = NAME_None;
+	FName Class = NAME_None;
 
-				if (MemoryUsageQueries::GetMemoryUsageCombined(MemoryUsageQueries::CurrentMemoryUsageInfoProvider, Packages, TotalSize, &Ar))
-				{
-					OutputDevice->Logf(TEXT("MemoryUsageCombined: TotalSize: %.2f MiB (%.2f KiB)"), TotalSize / (1024.f * 1024.f), TotalSize / 1024.f);
-				}
+	FString Name;
+	FString Names;
+	FString AssetName;
+	FString LogFileName;
 
-				bResult = true;
-			}
+	explicit FCommonParameters(const TCHAR* Args)
+	{
+		bTruncate = !FParse::Param(Args, TEXT("notrunc"));
+		bCSV = FParse::Param(Args, TEXT("csv"));
+
+		FParse::Value(Args, TEXT("Name="), Name);
+		FParse::Value(Args, TEXT("Names="), Names);
+		FParse::Value(Args, TEXT("Limit="), Limit);
+		FParse::Value(Args, TEXT("Log="), LogFileName);
+		FParse::Value(Args, TEXT("Asset="), AssetName);
+
+		FString GroupName;
+		if (FParse::Value(Args, TEXT("Group="), GroupName))
+		{
+			Group = FName(*GroupName);
 		}
-		else if (FParse::Command(&Cmd, TEXT("SharedUsage")))
+
+		FString ClassName;
+		if (FParse::Value(Args, TEXT("Class="), ClassName))
 		{
-			if (!Names.IsEmpty())
-			{
-				TArray<FString> Packages;
-				Names.ParseIntoArrayWS(Packages);
-				uint64 SharedSize;
+			Class = FName(*ClassName);
+		}
+	}
 
-				if (MemoryUsageQueries::GetMemoryUsageShared(MemoryUsageQueries::CurrentMemoryUsageInfoProvider, Packages, SharedSize, &Ar))
-				{
-					OutputDevice->Logf(TEXT("MemoryUsageShared: SharedSize: %.2f MiB (%.2f KiB)"), SharedSize / (1024.f * 1024.f), SharedSize / 1024.f);
-				}
+	FCommonParameters() = delete;
+};
 
-				bResult = true;
-			}
+FAutoConsoleCommandWithWorldArgsAndOutputDevice GMemQueryUsage(
+	TEXT("MemQuery.Usage"),
+	TEXT("Name=<AssetName> Prints memory usage of the specified asset."),
+	FConsoleCommandWithWorldArgsAndOutputDeviceDelegate::CreateStatic(
+		[](const TArray<FString>& Args, UWorld* World, FOutputDevice& Ar)
+{
+	const FString Cmd = FString::Join(Args, TEXT(" "));
+
+	FScopedOutputDevice ScopedOutputDevice(&Ar);
+	FCommonParameters CommonArgs(*Cmd);
+
+	ScopedOutputDevice.OpenLogFile(CommonArgs.LogFileName, CommonArgs.bCSV);
+
+	if (!CommonArgs.Name.IsEmpty())
+	{
+		uint64 ExclusiveSize = 0U;
+		uint64 InclusiveSize = 0U;
+
+		if (MemoryUsageQueries::GetMemoryUsage(MemoryUsageQueries::CurrentMemoryUsageInfoProvider, CommonArgs.Name, ExclusiveSize, InclusiveSize, &Ar))
+		{
+			ScopedOutputDevice.GetOutputDevice().Logf(TEXT("MemoryUsage: ExclusiveSize: %.2f MiB (%.2f KiB); InclusiveSize: %.2f MiB (%.2f KiB)"),
+				ExclusiveSize / (1024.f * 1024.f),
+				ExclusiveSize / 1024.f,
+				InclusiveSize / (1024.f * 1024.f),
+				InclusiveSize / 1024.f);
 		}
-		else if (FParse::Command(&Cmd, TEXT("UniqueUsage")))
+	}
+}));
+
+FAutoConsoleCommandWithWorldArgsAndOutputDevice GMemQueryCombinedUsage(
+	TEXT("MemQuery.CombinedUsage"),
+	TEXT("Names=\"<AssetName1> <AssetName2> ...\" Prints combined memory usage of the specified assets (including all dependencies)."),
+	FConsoleCommandWithWorldArgsAndOutputDeviceDelegate::CreateStatic(
+		[](const TArray<FString>& Args, UWorld* World, FOutputDevice& Ar)
+{
+	const FString Cmd = FString::Join(Args, TEXT(" "));
+
+	FScopedOutputDevice ScopedOutputDevice(&Ar);
+	FCommonParameters CommonArgs(*Cmd);
+
+	ScopedOutputDevice.OpenLogFile(CommonArgs.LogFileName, CommonArgs.bCSV);
+
+	if (!CommonArgs.Names.IsEmpty())
+	{
+		TArray<FString> Packages;
+		CommonArgs.Names.ParseIntoArrayWS(Packages);
+		uint64 TotalSize;
+
+		if (MemoryUsageQueries::GetMemoryUsageCombined(MemoryUsageQueries::CurrentMemoryUsageInfoProvider, Packages, TotalSize, &Ar))
 		{
-			if (!Names.IsEmpty())
-			{
-				TArray<FString> Packages;
-				Names.ParseIntoArrayWS(Packages);
-				uint64 UniqueSize = 0U;
+			ScopedOutputDevice.GetOutputDevice().Logf(TEXT("MemoryUsageCombined: TotalSize: %.2f MiB (%.2f KiB)"), TotalSize / (1024.f * 1024.f), TotalSize / 1024.f);
+		}
+	}
+}));
 
-				if (MemoryUsageQueries::GetMemoryUsageUnique(MemoryUsageQueries::CurrentMemoryUsageInfoProvider, Packages, UniqueSize, &Ar))
-				{
-					OutputDevice->Logf(TEXT("MemoryUsageUnique: UniqueSize: %.2f MiB (%.2f KiB)"), UniqueSize / (1024.f * 1024.f), UniqueSize / 1024.f);
-				}
+FAutoConsoleCommandWithWorldArgsAndOutputDevice GMemQuerySharedUsage(
+	TEXT("MemQuery.SharedUsage"),
+	TEXT("Names=\"<AssetName1> <AssetName2> ...\" Prints shared memory usage of the specified assets (including only dependencies shared by the specified assets)."),
+	FConsoleCommandWithWorldArgsAndOutputDeviceDelegate::CreateStatic(
+		[](const TArray<FString>& Args, UWorld* World, FOutputDevice& Ar)
+{
+	const FString Cmd = FString::Join(Args, TEXT(" "));
 
-				bResult = true;
-			}
+	FScopedOutputDevice ScopedOutputDevice(&Ar);
+	FCommonParameters CommonArgs(*Cmd);
+
+	ScopedOutputDevice.OpenLogFile(CommonArgs.LogFileName, CommonArgs.bCSV);
+
+	if (!CommonArgs.Names.IsEmpty())
+	{
+		TArray<FString> Packages;
+		CommonArgs.Names.ParseIntoArrayWS(Packages);
+		uint64 SharedSize;
+
+		if (MemoryUsageQueries::GetMemoryUsageShared(MemoryUsageQueries::CurrentMemoryUsageInfoProvider, Packages, SharedSize, &Ar))
+		{
+			ScopedOutputDevice.GetOutputDevice().Logf(TEXT("MemoryUsageShared: SharedSize: %.2f MiB (%.2f KiB)"), SharedSize / (1024.f * 1024.f), SharedSize / 1024.f);
 		}
-		else if (FParse::Command(&Cmd, TEXT("CommonUsage")))
+	}
+}));
+
+FAutoConsoleCommandWithWorldArgsAndOutputDevice GMemQueryUniqueUsage(
+	TEXT("MemQuery.UniqueUsage"),
+	TEXT("Names=\"<AssetName1> <AssetName2> ...\" Prints unique memory usage of the specified assets (including only dependencies unique to the specified assets)."),
+	FConsoleCommandWithWorldArgsAndOutputDeviceDelegate::CreateStatic(
+		[](const TArray<FString>& Args, UWorld* World, FOutputDevice& Ar)
+{
+	const FString Cmd = FString::Join(Args, TEXT(" "));
+
+	FScopedOutputDevice ScopedOutputDevice(&Ar);
+	FCommonParameters CommonArgs(*Cmd);
+
+	ScopedOutputDevice.OpenLogFile(CommonArgs.LogFileName, CommonArgs.bCSV);
+
+	if (!CommonArgs.Names.IsEmpty())
+	{
+		TArray<FString> Packages;
+		CommonArgs.Names.ParseIntoArrayWS(Packages);
+		uint64 UniqueSize = 0U;
+
+		if (MemoryUsageQueries::GetMemoryUsageUnique(MemoryUsageQueries::CurrentMemoryUsageInfoProvider, Packages, UniqueSize, &Ar))
 		{
-			if (!Names.IsEmpty())
-			{
-				TArray<FString> Packages;
-				Names.ParseIntoArrayWS(Packages);
-				uint64 CommonSize = 0U;
+			ScopedOutputDevice.GetOutputDevice().Logf(TEXT("MemoryUsageUnique: UniqueSize: %.2f MiB (%.2f KiB)"), UniqueSize / (1024.f * 1024.f), UniqueSize / 1024.f);
+		}
+	}
+}));
 
-				if (MemoryUsageQueries::GetMemoryUsageCommon(MemoryUsageQueries::CurrentMemoryUsageInfoProvider, Packages, CommonSize, &Ar))
-				{
-					OutputDevice->Logf(TEXT("MemoryUsageCommon: CommonSize: %.2f MiB (%.2f KiB)"), CommonSize / (1024.f * 1024.f), CommonSize / 1024.f);
-				}
+FAutoConsoleCommandWithWorldArgsAndOutputDevice GMemQueryCommonUsage(
+	TEXT("MemQuery.CommonUsage"),
+	TEXT("Names=\"<AssetName1> <AssetName2> ...\" Prints common memory usage of the specified assets (including only dependencies that are not unique to the specified assets)."),
+	FConsoleCommandWithWorldArgsAndOutputDeviceDelegate::CreateStatic(
+		[](const TArray<FString>& Args, UWorld* World, FOutputDevice& Ar)
+{
+	const FString Cmd = FString::Join(Args, TEXT(" "));
 
-				bResult = true;
-			}
+	FScopedOutputDevice ScopedOutputDevice(&Ar);
+	FCommonParameters CommonArgs(*Cmd);
+
+	ScopedOutputDevice.OpenLogFile(CommonArgs.LogFileName, CommonArgs.bCSV);
+
+	if (!CommonArgs.Names.IsEmpty())
+	{
+		TArray<FString> Packages;
+		CommonArgs.Names.ParseIntoArrayWS(Packages);
+		uint64 CommonSize = 0U;
+
+		if (MemoryUsageQueries::GetMemoryUsageCommon(MemoryUsageQueries::CurrentMemoryUsageInfoProvider, Packages, CommonSize, &Ar))
+		{
+			ScopedOutputDevice.GetOutputDevice().Logf(TEXT("MemoryUsageCommon: CommonSize: %.2f MiB (%.2f KiB)"), CommonSize / (1024.f * 1024.f), CommonSize / 1024.f);
 		}
-		else if (FParse::Command(&Cmd, TEXT("Dependencies")))
+
+	}
+}));
+
+FAutoConsoleCommandWithWorldArgsAndOutputDevice GMemQueryDependencies(
+	TEXT("MemQuery.Dependencies"),
+	TEXT("Name=<AssetName> Limit=<n> Lists dependencies of the specified asset, sorted by size."),
+	FConsoleCommandWithWorldArgsAndOutputDeviceDelegate::CreateStatic(
+		[](const TArray<FString>& Args, UWorld* World, FOutputDevice& Ar)
+{
+	const FString Cmd = FString::Join(Args, TEXT(" "));
+
+	FScopedOutputDevice ScopedOutputDevice(&Ar);
+	FCommonParameters CommonArgs(*Cmd);
+
+	ScopedOutputDevice.OpenLogFile(CommonArgs.LogFileName, CommonArgs.bCSV);
+
+	if (!CommonArgs.Name.IsEmpty())
+	{
+		TMap<FName, uint64> DependenciesWithSize;
+
+		if (MemoryUsageQueries::GetDependenciesWithSize(MemoryUsageQueries::CurrentMemoryUsageInfoProvider, CommonArgs.Name, DependenciesWithSize, &Ar))
 		{
-			if (!Name.IsEmpty())
-			{
-				TMap<FName, uint64> DependenciesWithSize;
+			MemoryUsageQueries::Internal::PrintTagsWithSize(ScopedOutputDevice.GetOutputDevice(), DependenciesWithSize, TEXT("Dependencies"), CommonArgs.bTruncate, CommonArgs.Limit, CommonArgs.bCSV);
+		}
+	}
+}));
 
-				if (MemoryUsageQueries::GetDependenciesWithSize(MemoryUsageQueries::CurrentMemoryUsageInfoProvider, Name, DependenciesWithSize, &Ar))
-				{
-					MemoryUsageQueries::Internal::PrintTagsWithSize(*OutputDevice, DependenciesWithSize, TEXT("dependencies"), bTruncate, Limit, bCSV);
-				}
+FAutoConsoleCommandWithWorldArgsAndOutputDevice GMemQueryCombinedDependencies(
+	TEXT("MemQuery.CombinedDependencies"),
+	TEXT("Names=\"<AssetName1> <AssetName2> ...\" Limit=<n> Lists n largest dependencies of the specified assets, sorted by size."),
+	FConsoleCommandWithWorldArgsAndOutputDeviceDelegate::CreateStatic(
+		[](const TArray<FString>& Args, UWorld* World, FOutputDevice& Ar)
+{
+	const FString Cmd = FString::Join(Args, TEXT(" "));
 
-				bResult = true;
-			}
+	FScopedOutputDevice ScopedOutputDevice(&Ar);
+	FCommonParameters CommonArgs(*Cmd);
+
+	ScopedOutputDevice.OpenLogFile(CommonArgs.LogFileName, CommonArgs.bCSV);
+
+	if (!CommonArgs.Names.IsEmpty())
+	{
+		TArray<FString> Packages;
+		CommonArgs.Names.ParseIntoArrayWS(Packages);
+		TMap<FName, uint64> CombinedDependenciesWithSize;
+
+		if (MemoryUsageQueries::GetDependenciesWithSizeCombined(MemoryUsageQueries::CurrentMemoryUsageInfoProvider, Packages, CombinedDependenciesWithSize, &Ar))
+		{
+			MemoryUsageQueries::Internal::PrintTagsWithSize(ScopedOutputDevice.GetOutputDevice(), CombinedDependenciesWithSize, TEXT("Combined Dependencies"), CommonArgs.bTruncate, CommonArgs.Limit, CommonArgs.bCSV);
 		}
-		else if (FParse::Command(&Cmd, TEXT("CombinedDependencies")))
+	}
+}));
+
+FAutoConsoleCommandWithWorldArgsAndOutputDevice GMemQuerySharedDependencies(
+	TEXT("MemQuery.SharedDependencies"),
+	TEXT("Names=\"<AssetName1> <AssetName2> ...\" Limit=<n> Lists n largest dependencies that are shared by the specified assets, sorted by size."),
+	FConsoleCommandWithWorldArgsAndOutputDeviceDelegate::CreateStatic(
+		[](const TArray<FString>& Args, UWorld* World, FOutputDevice& Ar)
+{
+	const FString Cmd = FString::Join(Args, TEXT(" "));
+
+	FScopedOutputDevice ScopedOutputDevice(&Ar);
+	FCommonParameters CommonArgs(*Cmd);
+
+	ScopedOutputDevice.OpenLogFile(CommonArgs.LogFileName, CommonArgs.bCSV);
+
+	if (!CommonArgs.Names.IsEmpty())
+	{
+		TArray<FString> Packages;
+		CommonArgs.Names.ParseIntoArrayWS(Packages);
+		TMap<FName, uint64> SharedDependenciesWithSize;
+
+		if (MemoryUsageQueries::GetDependenciesWithSizeShared(MemoryUsageQueries::CurrentMemoryUsageInfoProvider, Packages, SharedDependenciesWithSize, &Ar))
 		{
-			if (!Names.IsEmpty())
-			{
-				TArray<FString> Packages;
-				Names.ParseIntoArrayWS(Packages);
-				TMap<FName, uint64> CombinedDependenciesWithSize;
+			MemoryUsageQueries::Internal::PrintTagsWithSize(ScopedOutputDevice.GetOutputDevice(), SharedDependenciesWithSize, TEXT("Shared Dependencies"), CommonArgs.bTruncate, CommonArgs.Limit, CommonArgs.bCSV);
+		}
+	}
+}));
 
-				if (MemoryUsageQueries::GetDependenciesWithSizeCombined(MemoryUsageQueries::CurrentMemoryUsageInfoProvider, Packages, CombinedDependenciesWithSize, &Ar))
-				{
-					MemoryUsageQueries::Internal::PrintTagsWithSize(*OutputDevice, CombinedDependenciesWithSize, TEXT("combined dependencies"), bTruncate, Limit, bCSV);
-				}
+FAutoConsoleCommandWithWorldArgsAndOutputDevice GMemQueryUniqueDependencies(
+	TEXT("MemQuery.UniqueDependencies"),
+	TEXT("Names=\"<AssetName1> <AssetName2> ...\" Limit=<n> Lists n largest dependencies that are unique to the specified assets, sorted by size."),
+	FConsoleCommandWithWorldArgsAndOutputDeviceDelegate::CreateStatic(
+		[](const TArray<FString>& Args, UWorld* World, FOutputDevice& Ar)
+{
+	const FString Cmd = FString::Join(Args, TEXT(" "));
 
-				bResult = true;
-			}
+	FScopedOutputDevice ScopedOutputDevice(&Ar);
+	FCommonParameters CommonArgs(*Cmd);
+
+	ScopedOutputDevice.OpenLogFile(CommonArgs.LogFileName, CommonArgs.bCSV);
+
+	if (!CommonArgs.Names.IsEmpty())
+	{
+		TArray<FString> Packages;
+		CommonArgs.Names.ParseIntoArrayWS(Packages);
+		TMap<FName, uint64> UniqueDependenciesWithSize;
+
+		if (MemoryUsageQueries::GetDependenciesWithSizeUnique(MemoryUsageQueries::CurrentMemoryUsageInfoProvider, Packages, UniqueDependenciesWithSize, &Ar))
+		{
+			MemoryUsageQueries::Internal::PrintTagsWithSize(ScopedOutputDevice.GetOutputDevice(), UniqueDependenciesWithSize, TEXT("Unique Dependencies"), CommonArgs.bTruncate, CommonArgs.Limit, CommonArgs.bCSV);
 		}
-		else if (FParse::Command(&Cmd, TEXT("SharedDependencies")))
+	}
+}));
+
+FAutoConsoleCommandWithWorldArgsAndOutputDevice GMemQueryCommonDependencies(
+	TEXT("MemQuery.CommonDependencies"),
+	TEXT("Names=\"<AssetName1> <AssetName2> ...\" Limit=<n> Lists n largest dependencies that are NOT unique to the specified assets, sorted by size."),
+	FConsoleCommandWithWorldArgsAndOutputDeviceDelegate::CreateStatic(
+		[](const TArray<FString>& Args, UWorld* World, FOutputDevice& Ar)
+{
+	const FString Cmd = FString::Join(Args, TEXT(" "));
+
+	FScopedOutputDevice ScopedOutputDevice(&Ar);
+	FCommonParameters CommonArgs(*Cmd);
+
+	ScopedOutputDevice.OpenLogFile(CommonArgs.LogFileName, CommonArgs.bCSV);
+
+	if (!CommonArgs.Names.IsEmpty())
+	{
+		TArray<FString> Packages;
+		CommonArgs.Names.ParseIntoArrayWS(Packages);
+		TMap<FName, uint64> CommonDependenciesWithSize;
+
+		if (MemoryUsageQueries::GetDependenciesWithSizeCommon(MemoryUsageQueries::CurrentMemoryUsageInfoProvider, Packages, CommonDependenciesWithSize, &Ar))
 		{
-			if (!Names.IsEmpty())
-			{
-				TArray<FString> Packages;
-				Names.ParseIntoArrayWS(Packages);
-				TMap<FName, uint64> SharedDependenciesWithSize;
+			MemoryUsageQueries::Internal::PrintTagsWithSize(ScopedOutputDevice.GetOutputDevice(), CommonDependenciesWithSize, TEXT("Common Dependencies"), CommonArgs.bTruncate, CommonArgs.Limit, CommonArgs.bCSV);
+		}
+	}
+}));
 
-				if (MemoryUsageQueries::GetDependenciesWithSizeShared(MemoryUsageQueries::CurrentMemoryUsageInfoProvider, Packages, SharedDependenciesWithSize, &Ar))
-				{
-					MemoryUsageQueries::Internal::PrintTagsWithSize(*OutputDevice, SharedDependenciesWithSize, TEXT("shared dependencies"), bTruncate, Limit, bCSV);
-				}
+FAutoConsoleCommandWithWorldArgsAndOutputDevice GMemQuerySavings(
+	TEXT("MemQuery.Savings"),
+	TEXT("Limit=<n> Lists potential savings among %s. How much memory can be saved it we delete certain object."),
+	FConsoleCommandWithWorldArgsAndOutputDeviceDelegate::CreateStatic(
+		[](const TArray<FString>& Args, UWorld* World, FOutputDevice& Ar)
+{
+	const FString Cmd = FString::Join(Args, TEXT(" "));
 
-				bResult = true;
-			}
+	FScopedOutputDevice ScopedOutputDevice(&Ar);
+	FCommonParameters CommonArgs(*Cmd);
+
+	ScopedOutputDevice.OpenLogFile(CommonArgs.LogFileName, CommonArgs.bCSV);
+
+	const UMemoryUsageQueriesConfig* Config = GetDefault<UMemoryUsageQueriesConfig>();
+
+	// Necessary for FParse::Command
+	const TCHAR* Command = *Cmd;
+
+	for (auto It = Config->SavingsPresets.CreateConstIterator(); It; ++It)
+	{
+		if (!FParse::Command(&Command, *It.Key()))
+		{
+			continue;
 		}
-		else if (FParse::Command(&Cmd, TEXT("UniqueDependencies")))
+
+		TMap<FName, uint64> PresetSavings;
+		TSet<FName> Packages;
+
+		UClass* SavingsClass = FindObject<UClass>(nullptr, *It.Value());
+		if (SavingsClass != nullptr)
 		{
-			if (!Names.IsEmpty())
-			{
-				TArray<FString> Packages;
-				Names.ParseIntoArrayWS(Packages);
-				TMap<FName, uint64> UniqueDependenciesWithSize;
+			TArray<UClass*> DerivedClasses;
+			TArray<UClass*> DerivedResults;
+
+			GetDerivedClasses(SavingsClass, DerivedClasses, true);
 
-				if (MemoryUsageQueries::GetDependenciesWithSizeUnique(MemoryUsageQueries::CurrentMemoryUsageInfoProvider, Packages, UniqueDependenciesWithSize, &Ar))
+			for (UClass* DerivedClass : DerivedClasses)
+			{
+				UBlueprintGeneratedClass* BPClass = Cast<UBlueprintGeneratedClass>(DerivedClass);
+				if (BPClass != nullptr)
 				{
-					MemoryUsageQueries::Internal::PrintTagsWithSize(*OutputDevice, UniqueDependenciesWithSize, TEXT("unique dependencies"), bTruncate, Limit, bCSV);
+					DerivedResults.Reset();
+					GetDerivedClasses(BPClass, DerivedResults, false);
+					if (DerivedResults.Num() == 0)
+					{
+						Packages.Add(DerivedClass->GetPackage()->GetFName());
+					}
 				}
-
-				bResult = true;
 			}
 		}
-		else if (FParse::Command(&Cmd, TEXT("CommonDependencies")))
+
+		for (const auto& Package : Packages)
 		{
-			if (!Names.IsEmpty())
-			{
-				TArray<FString> Packages;
-				Names.ParseIntoArrayWS(Packages);
-				TMap<FName, uint64> CommonDependenciesWithSize;
+			uint64 Size = 0;
+			MemoryUsageQueries::GetMemoryUsageUnique(MemoryUsageQueries::CurrentMemoryUsageInfoProvider, TArray<FString>({ Package.ToString() }), Size, &Ar);
+			PresetSavings.Add(Package, Size);
+		}
 
-				if (MemoryUsageQueries::GetDependenciesWithSizeCommon(MemoryUsageQueries::CurrentMemoryUsageInfoProvider, Packages, CommonDependenciesWithSize, &Ar))
-				{
-					MemoryUsageQueries::Internal::PrintTagsWithSize(*OutputDevice, CommonDependenciesWithSize, TEXT("common dependencies"), bTruncate, Limit, bCSV);
-				}
+		PresetSavings.ValueSort(TGreater<uint64>());
+		MemoryUsageQueries::Internal::PrintTagsWithSize(ScopedOutputDevice.GetOutputDevice(), PresetSavings, *FString::Printf(TEXT("possible savings")), CommonArgs.bTruncate, CommonArgs.bCSV);
+	}
+}));
 
-				bResult = true;
-			}
-		}
 #if ENABLE_LOW_LEVEL_MEM_TRACKER
-		else if (FParse::Command(&Cmd, TEXT("ListAssets")))
-		{
-			FString AssetName;
-			FParse::Value(Cmd, TEXT("NAME="), AssetName);
 
-			FName Group = NAME_None;
-			FString GroupName;
-			if (FParse::Value(Cmd, TEXT("GROUP="), GroupName))
-			{
-				Group = FName(*GroupName);
-			}
+FAutoConsoleCommandWithWorldArgsAndOutputDevice GMemQueryListAssets(
+	TEXT("MemQuery.ListAssets"),
+	TEXT("Name=<AssetNameSubstring> Group=<GroupName> Class=<ClassName> Limit=<n> Lists n largest assets."),
+	FConsoleCommandWithWorldArgsAndOutputDeviceDelegate::CreateStatic(
+		[](const TArray<FString>& Args, UWorld* World, FOutputDevice& Ar)
+{
+	const FString Cmd = FString::Join(Args, TEXT(" "));
 
-			FName Class = NAME_None;
-			FString ClassName;
-			if (FParse::Value(Cmd, TEXT("CLASS="), ClassName))
-			{
-				Class = FName(*ClassName);
-			}
+	FScopedOutputDevice ScopedOutputDevice(&Ar);
+	FCommonParameters CommonArgs(*Cmd);
 
-			bool bSuccess;
-			TMap<FName, uint64> AssetsWithSize;
+	ScopedOutputDevice.OpenLogFile(CommonArgs.LogFileName, CommonArgs.bCSV);
 
-			if (Group != NAME_None || Class != NAME_None)
-			{
-				bSuccess = MemoryUsageQueries::GetFilteredPackagesWithSize(AssetsWithSize, Group, AssetName, Class, &Ar);
-			}
-			else
-			{
-				// TODO - Implement using faster path if there are no group / class filters
-				bSuccess = MemoryUsageQueries::GetFilteredPackagesWithSize(AssetsWithSize, Group, AssetName, Class, &Ar);
-			}
+	bool bSuccess;
+	TMap<FName, uint64> AssetsWithSize;
 
-			if (bSuccess)
-			{
-				MemoryUsageQueries::Internal::PrintTagsWithSize(*OutputDevice, AssetsWithSize, TEXT("largest assets"), bTruncate, Limit, bCSV);
-			}
+	if (CommonArgs.Group != NAME_None || CommonArgs.Class != NAME_None)
+	{
+		bSuccess = MemoryUsageQueries::GetFilteredPackagesWithSize(AssetsWithSize, CommonArgs.Group, CommonArgs.AssetName, CommonArgs.Class, &ScopedOutputDevice.GetOutputDevice());
+	}
+	else
+	{
+		// TODO - Implement using faster path if there are no group / class filters
+		bSuccess = MemoryUsageQueries::GetFilteredPackagesWithSize(AssetsWithSize, CommonArgs.Group, CommonArgs.AssetName, CommonArgs.Class, &ScopedOutputDevice.GetOutputDevice());
+	}
 
-			bResult = true;
-		}
-		else if (FParse::Command(&Cmd, TEXT("ListClasses")))
-		{
-			FName Group = NAME_None;
-			FString GroupName;
-			if (FParse::Value(Cmd, TEXT("GROUP="), GroupName))
-			{
-				Group = FName(*GroupName);
-			}
+	if (bSuccess)
+	{
+		MemoryUsageQueries::Internal::PrintTagsWithSize(ScopedOutputDevice.GetOutputDevice(), AssetsWithSize, TEXT("largest assets"), CommonArgs.bTruncate, CommonArgs.Limit, CommonArgs.bCSV);
+	}
+}));
 
-			FString AssetName;
-			FParse::Value(Cmd, TEXT("ASSET="), AssetName);
-			
-			TMap<FName, uint64> ClassesWithSize;
+FAutoConsoleCommandWithWorldArgsAndOutputDevice GMemQueryListClasses(
+	TEXT("MemQuery.ListClasses"),
+	TEXT("Group=<GroupName> Asset=<AssetName> Limit=<n> Lists n largest classes."),
+	FConsoleCommandWithWorldArgsAndOutputDeviceDelegate::CreateStatic(
+		[](const TArray<FString>& Args, UWorld* World, FOutputDevice& Ar)
+{
+	const FString Cmd = FString::Join(Args, TEXT(" "));
 
-			if (MemoryUsageQueries::GetFilteredClassesWithSize(ClassesWithSize, Group, AssetName, &Ar))
-			{
-				MemoryUsageQueries::Internal::PrintTagsWithSize(*OutputDevice, ClassesWithSize, *FString::Printf(TEXT("largest classes")), bTruncate, Limit, bCSV);
-			}
+	FScopedOutputDevice ScopedOutputDevice(&Ar);
+	FCommonParameters CommonArgs(*Cmd);
 
-			bResult = true;
-		}
-		else if (FParse::Command(&Cmd, TEXT("ListGroups")))
-		{
-			FString AssetName = "";
-			FParse::Value(Cmd, TEXT("ASSET="), AssetName);
+	ScopedOutputDevice.OpenLogFile(CommonArgs.LogFileName, CommonArgs.bCSV);
 
-			FName Class = NAME_None;
-			FString ClassName;
-			if (FParse::Value(Cmd, TEXT("CLASS="), ClassName))
-			{
-				Class = FName(*ClassName);
-			}
+	TMap<FName, uint64> ClassesWithSize;
+	if (MemoryUsageQueries::GetFilteredClassesWithSize(ClassesWithSize, CommonArgs.Group, CommonArgs.AssetName, &Ar))
+	{
+		MemoryUsageQueries::Internal::PrintTagsWithSize(ScopedOutputDevice.GetOutputDevice(), ClassesWithSize, *FString::Printf(TEXT("Largest Classes")), CommonArgs.bTruncate, CommonArgs.Limit, CommonArgs.bCSV);
+	}
+}));
 
-			TMap<FName, uint64> GroupsWithSize;
+FAutoConsoleCommandWithWorldArgsAndOutputDevice GMemQueryListGroups(
+	TEXT("MemQuery.ListGroups"),
+	TEXT("Asset=<AssetName> Class=<ClassName> Limit=<n> Lists n largest groups."),
+	FConsoleCommandWithWorldArgsAndOutputDeviceDelegate::CreateStatic(
+		[](const TArray<FString>& Args, UWorld* World, FOutputDevice& Ar)
+{
+	const FString Cmd = FString::Join(Args, TEXT(" "));
 
-			if (MemoryUsageQueries::GetFilteredGroupsWithSize(GroupsWithSize, AssetName, Class, &Ar))
-			{
-				MemoryUsageQueries::Internal::PrintTagsWithSize(*OutputDevice, GroupsWithSize, *FString::Printf(TEXT("largest groups")), bTruncate, Limit, bCSV);
-			}
+	FScopedOutputDevice ScopedOutputDevice(&Ar);
+	FCommonParameters CommonArgs(*Cmd);
+
+	ScopedOutputDevice.OpenLogFile(CommonArgs.LogFileName, CommonArgs.bCSV);
+
+	TMap<FName, uint64> GroupsWithSize;
+	if (MemoryUsageQueries::GetFilteredGroupsWithSize(GroupsWithSize, CommonArgs.AssetName, CommonArgs.Class, &Ar))
+	{
+		MemoryUsageQueries::Internal::PrintTagsWithSize(ScopedOutputDevice.GetOutputDevice(), GroupsWithSize, *FString::Printf(TEXT("Largest Groups")), CommonArgs.bTruncate, CommonArgs.Limit, CommonArgs.bCSV);
+	}
+}));
+
+FAutoConsoleCommandWithWorldArgsAndOutputDevice GMemQueryCollections(
+	TEXT("MemQuery.Collection"),
+	TEXT("Lists memory used by a collection. Can show dependency breakdown. Pass -showdeps to list dependencies."),
+	FConsoleCommandWithWorldArgsAndOutputDeviceDelegate::CreateStatic(
+		[](const TArray<FString>& Args, UWorld* World, FOutputDevice& Ar)
+{
+	const FString Cmd = FString::Join(Args, TEXT(" "));
+
+	FScopedOutputDevice ScopedOutputDevice(&Ar);
+	FCommonParameters CommonArgs(*Cmd);
+
+	ScopedOutputDevice.OpenLogFile(CommonArgs.LogFileName, CommonArgs.bCSV);
+
+	const bool bShowDependencies = FParse::Param(*Cmd, TEXT("ShowDeps"));
+
+	// Necessary for FParse::Command
+	const TCHAR* Command = *Cmd;
+
+	FOutputDevice* CurrentOutputDevice = &ScopedOutputDevice.GetOutputDevice();
 
-			bResult = true;
+	const UMemoryUsageQueriesConfig* Config = GetDefault<UMemoryUsageQueriesConfig>();
+	for (auto It = Config->Collections.CreateConstIterator(); It; ++It)
+	{
+		const FCollectionInfo& CollectionInfo = *It;
+
+		if (!FParse::Command(&Command, *CollectionInfo.Name))
+		{
+			continue;
 		}
-#endif
-		else if (FParse::Command(&Cmd, TEXT("Savings")))
+
+		// Retrieve a list of all assets that have allocations we are currently tracking.
+		TMap<FName, uint64> AssetsWithSize;
+		bool bSuccess = MemoryUsageQueries::GetFilteredPackagesWithSize(AssetsWithSize, NAME_None, "", NAME_None, &Ar);
+
+		if (!bSuccess)
 		{
-			const UMemoryUsageQueriesConfig* Config = GetDefault<UMemoryUsageQueriesConfig>();
+			Ar.Logf(TEXT("Failed to gather assets for Collection %s"), *CollectionInfo.Name);
+			break;
+		}
 
-			for (auto It = Config->SavingsPresets.CreateConstIterator(); It; ++It)
+		// Will return true if the Package Name matches any of the conditions in the array of Paths
+		auto PackageNameMatches = [](const FString& PackageName, const TArray<FString>& Conditions)->bool {
+			for (const FString& Condition : Conditions)
 			{
-				if (!FParse::Command(&Cmd, *It.Key()))
+				if ((FWildcardString::ContainsWildcards(*Condition) && FWildcardString::IsMatch(*Condition, *PackageName)) || (PackageName.Contains(Condition)))
 				{
-					continue;
+					return true;
 				}
+			}
 
-				TMap<FName, uint64> PresetSavings;
-				TSet<FName> Packages;
+			return false;
+		};
 
-				UClass* SavingsClass = FindObject<UClass>(nullptr, *It.Value());
-				if (SavingsClass != nullptr)
-				{
-					TArray<UClass*> DerivedClasses;
-					TArray<UClass*> DerivedResults;
+		// See if any of the asset paths match those of our matching paths and are valid
+		TArray<FString> PackageNames;
 
-					GetDerivedClasses(SavingsClass, DerivedClasses, true);
+		TMap<FName, FAssetMemoryDetails> AssetMemoryMap;
+		for (const TPair<FName, uint64>& AssetSizePair : AssetsWithSize)
+		{
+			const FString& PackageName = AssetSizePair.Key.ToString();
 
-					for (UClass* DerivedClass : DerivedClasses)
-					{
-						UBlueprintGeneratedClass* BPClass = Cast<UBlueprintGeneratedClass>(DerivedClass);
-						if (BPClass != nullptr)
-						{
-							DerivedResults.Reset();
-							GetDerivedClasses(BPClass, DerivedResults, false);
-							if (DerivedResults.Num() == 0)
-							{
-								Packages.Add(DerivedClass->GetPackage()->GetFName());
-							}
-						}
-					}
-				}
+			if (!FPackageName::IsValidLongPackageName(PackageName))
+			{
+				continue;
+			}
+
+			// If path is Included and NOT Excluded, it's a valid asset to consider
+			if (PackageNameMatches(PackageName, CollectionInfo.Includes) && !PackageNameMatches(PackageName, CollectionInfo.Excludes))
+			{
+				PackageNames.Add(PackageName);
+				FAssetMemoryDetails& AssetMemory = AssetMemoryMap.Add(AssetSizePair.Key);
+				AssetMemory.MemoryBreakdown.ExclusiveSize = AssetSizePair.Value;
 
-				for (const auto& Package : Packages)
+				FName LongPackageName;
+				if (!MemoryUsageQueries::Internal::GetLongNameAndDependencies(PackageName, LongPackageName, AssetMemory.Dependencies, &Ar))
 				{
-					uint64 Size = 0;
-					MemoryUsageQueries::GetMemoryUsageUnique(MemoryUsageQueries::CurrentMemoryUsageInfoProvider, TArray<FString>({Package.ToString()}), Size, &Ar);
-					PresetSavings.Add(Package, Size);
+					Ar.Logf(TEXT("Failed to get dependencies foro Asset %s"), *PackageName);
 				}
-
-				PresetSavings.ValueSort(TGreater<uint64>());
-				MemoryUsageQueries::Internal::PrintTagsWithSize(*OutputDevice, PresetSavings, *FString::Printf(TEXT("possible savings")), bTruncate, bCSV);
-
-				bResult = true;
 			}
 		}
-#if ENABLE_LOW_LEVEL_MEM_TRACKER
-		else if (FParse::Command(&Cmd, TEXT("Collection")))
+
+		// Gather list of dependencies. Internal dependencies are confined only to the set of packages passed in.
+		// External are dependencies that have additional references outside the set of packages passed in.
+		TMap<FName, uint64> InternalDependencies;
+		TMap<FName, uint64> ExternalDependencies;
+		if (!MemoryUsageQueries::GatherDependenciesForPackages(MemoryUsageQueries::CurrentMemoryUsageInfoProvider, PackageNames, InternalDependencies, ExternalDependencies, &Ar))
+		{
+			Ar.Logf(TEXT("Failed to gather memory usage for dependencies in Collection %s"), *CollectionInfo.Name);
+			break;
+		}
+
+		uint64 TotalCollectionSize = 0;
+
+		// Determine which category where each assets dependency should reside
+		for (TPair<FName, FAssetMemoryDetails>& Asset : AssetMemoryMap)
 		{
-			const bool bShowDependencies = FParse::Param(Cmd, TEXT("ShowDeps"));
+			FAssetMemoryDetails& AssetMemory = Asset.Value;
+			FAssetMemoryBreakdown& AssetMemoryDetails = AssetMemory.MemoryBreakdown;
 
-			const UMemoryUsageQueriesConfig* Config = GetDefault<UMemoryUsageQueriesConfig>();
-			for (auto It = Config->Collections.CreateConstIterator(); It; ++It)
+			for (FName& Dep : Asset.Value.Dependencies)
 			{
-				const FCollectionInfo& CollectionInfo = *It;
-				
-				if (!FParse::Command(&Cmd, *CollectionInfo.Name))
+				// Don't want to count asset itself, plus some dependencies might refer to other assets in the map
+				if (AssetMemoryMap.Contains(Dep))
 				{
 					continue;
 				}
-				
-				// Retrieve a list of all assets that have allocations we are currently tracking.
-				TMap<FName, uint64> AssetsWithSize;
-				bool bSuccess = MemoryUsageQueries::GetFilteredPackagesWithSize(AssetsWithSize, NAME_None, "", NAME_None, &Ar);
 
-				if (!bSuccess)
+				FAssetMemoryBreakdown DependencyMemory;
+				uint64* UniqueMemory = InternalDependencies.Find(Dep);
+				uint64* SharedMemory = ExternalDependencies.Find(Dep);
+				bool bRecordDependency = false;
+
+				if (UniqueMemory != nullptr && *UniqueMemory != 0)
 				{
-					Ar.Logf(TEXT("Failed to gather assets for Collection %s"), *CollectionInfo.Name);
-					break;
+					DependencyMemory.UniqueSize = *UniqueMemory;
+					AssetMemoryDetails.UniqueSize += DependencyMemory.UniqueSize;
+					AssetMemory.UniqueRefCount++;
+					bRecordDependency = true;
 				}
 
-				// Will return true if the Package Name matches any of the conditions in the array of Paths
-				auto PackageNameMatches = [](const FString& PackageName, const TArray<FString>& Conditions)->bool {
-					for (const FString& Condition : Conditions)
-					{
-						if ((FWildcardString::ContainsWildcards(*Condition) && FWildcardString::IsMatch(*Condition, *PackageName)) || (PackageName.Contains(Condition)))
-						{
-							return true;
-						}
-					}
-
-					return false;
-				};
-				
-				// See if any of the asset paths match those of our matching paths and are valid
-				TArray<FString> PackageNames;
-
-				TMap<FName, FAssetMemoryDetails> AssetMemoryMap;
-				for (const TPair<FName, uint64>& AssetSizePair : AssetsWithSize)
+				if (SharedMemory != nullptr && *SharedMemory != 0)
 				{
-					const FString& PackageName = AssetSizePair.Key.ToString();
-					
-					if (!FPackageName::IsValidLongPackageName(PackageName))
-					{
-						continue;
-					}
-
-					// If path is Included and NOT Excluded, it's a valid asset to consider
-					if (PackageNameMatches(PackageName, CollectionInfo.Includes) && !PackageNameMatches(PackageName, CollectionInfo.Excludes))
-					{
-						PackageNames.Add(PackageName);
-						FAssetMemoryDetails& AssetMemory = AssetMemoryMap.Add(AssetSizePair.Key);
-						AssetMemory.MemoryBreakdown.ExclusiveSize = AssetSizePair.Value;
-
-						FName LongPackageName;
-						if (!MemoryUsageQueries::Internal::GetLongNameAndDependencies(PackageName, LongPackageName, AssetMemory.Dependencies, &Ar))
-						{
-							Ar.Logf(TEXT("Failed to get dependencies foro Asset %s"), *PackageName);
-						}
-					}
+					DependencyMemory.SharedSize = *SharedMemory;
+					AssetMemoryDetails.SharedSize += DependencyMemory.SharedSize;
+					AssetMemory.SharedRefCount++;
+					bRecordDependency = true;
 				}
 
-				// Gather list of dependencies. Internal dependencies are confined only to the set of packages passed in.
-				// External are dependencies that have additional references outside the set of packages passed in.
-				TMap<FName, uint64> InternalDependencies;
-				TMap<FName, uint64> ExternalDependencies;
-				if (!MemoryUsageQueries::GatherDependenciesForPackages(MemoryUsageQueries::CurrentMemoryUsageInfoProvider, PackageNames, InternalDependencies, ExternalDependencies, &Ar))
+				if (bRecordDependency)
 				{
-					Ar.Logf(TEXT("Failed to gather memory usage for dependencies in Collection %s"), *CollectionInfo.Name);
-					break;
+					AssetMemory.DependenciesToMemoryMap.Add(Dep, DependencyMemory);
 				}
+			}
 
-				uint64 TotalCollectionSize = 0;
+			AssetMemoryDetails.TotalSize = AssetMemoryDetails.ExclusiveSize + AssetMemoryDetails.UniqueSize;
+			TotalCollectionSize += AssetMemoryDetails.TotalSize;
+		}
 
-				// Determine which category where each assets dependency should reside
-				for (TPair<FName, FAssetMemoryDetails>& Asset : AssetMemoryMap)
-				{
-					FAssetMemoryDetails& AssetMemory = Asset.Value;
-					FAssetMemoryBreakdown& AssetMemoryDetails = AssetMemory.MemoryBreakdown;
+		// Sort by TotalSize
+		AssetMemoryMap.ValueSort([](const FAssetMemoryDetails& A, const FAssetMemoryDetails& B) {
+			return A.MemoryBreakdown.TotalSize > B.MemoryBreakdown.TotalSize;
+			});
 
-					for (FName& Dep : Asset.Value.Dependencies)
-					{
-						// Don't want to count asset itself, plus some dependencies might refer to other assets in the map
-						if (AssetMemoryMap.Contains(Dep))
-						{
-							continue;
-						}
-
-						FAssetMemoryBreakdown DependencyMemory;
-						uint64* UniqueMemory = InternalDependencies.Find(Dep);
-						uint64* SharedMemory = ExternalDependencies.Find(Dep);
-						bool bRecordDependency = false;
-
-						if (UniqueMemory != nullptr && *UniqueMemory != 0)
-						{
-							DependencyMemory.UniqueSize = *UniqueMemory;
-							AssetMemoryDetails.UniqueSize += DependencyMemory.UniqueSize;
-							AssetMemory.UniqueRefCount++;
-							bRecordDependency = true;
-						}
-
-						if (SharedMemory != nullptr && *SharedMemory != 0)
-						{
-							DependencyMemory.SharedSize = *SharedMemory;
-							AssetMemoryDetails.SharedSize += DependencyMemory.SharedSize;
-							AssetMemory.SharedRefCount++;
-							bRecordDependency = true;
-						}
-
-						if (bRecordDependency)
-						{
-							AssetMemory.DependenciesToMemoryMap.Add(Dep, DependencyMemory);
-						}
-					}
+		if (CommonArgs.bCSV)
+		{
+			CurrentOutputDevice->Logf(TEXT(",Asset,Exclusive KiB,Unique Refs KiB,Unique Ref Count,Shared Refs KiB,Shared Ref Count,Total KiB"));
+		}
+		else
+		{
+			CurrentOutputDevice->Logf(
+				TEXT(" %100s %20s %20s %15s %20s %15s %25s"),
+				TEXT("Asset"),
+				TEXT("Exclusive KiB"),
+				TEXT("Unique Refs KiB"),
+				TEXT("Unique Ref Count"),
+				TEXT("Shared Refs KiB"),
+				TEXT("Shared Ref Count"),
+				TEXT("Total KiB")
+			);
+		}
 
-					AssetMemoryDetails.TotalSize = AssetMemoryDetails.ExclusiveSize + AssetMemoryDetails.UniqueSize;
-					TotalCollectionSize += AssetMemoryDetails.TotalSize;
-				}
+		// Asset listing
+		for (TPair<FName, FAssetMemoryDetails>& Asset : AssetMemoryMap)
+		{
+			FAssetMemoryDetails& AssetMemory = Asset.Value;
+			FAssetMemoryBreakdown& AssetMemoryDetails = AssetMemory.MemoryBreakdown;
 
-				// Sort by TotalSize
-				AssetMemoryMap.ValueSort([](const FAssetMemoryDetails& A, const FAssetMemoryDetails& B) {
-						return A.MemoryBreakdown.TotalSize > B.MemoryBreakdown.TotalSize;
-					});
+			if (CommonArgs.bCSV)
+			{
+				CurrentOutputDevice->Logf(TEXT(",%s,%.2f,%.2f,%d,%.2f,%d,%.2f"), *Asset.Key.ToString(),
+					AssetMemoryDetails.ExclusiveSize / 1024.f,
+					AssetMemoryDetails.UniqueSize / 1024.f,
+					AssetMemory.UniqueRefCount,
+					AssetMemoryDetails.SharedSize / 1024.f,
+					AssetMemory.SharedRefCount,
+					AssetMemoryDetails.TotalSize / 1024.f);
+			}
+			else
+			{
+				CurrentOutputDevice->Logf(
+					TEXT(" %100s %20.2f %20.2f %15d %20.2f %15d %25.2f"),
+					*Asset.Key.ToString(),
+					AssetMemoryDetails.ExclusiveSize / 1024.f,
+					AssetMemoryDetails.SharedSize / 1024.f,
+					AssetMemory.SharedRefCount,
+					AssetMemoryDetails.UniqueSize / 1024.f,
+					AssetMemory.UniqueRefCount,
+					AssetMemoryDetails.TotalSize / 1024.f
+				);
+			}
+		}
 
-				if (bCSV)
-				{
-					OutputDevice->Logf(TEXT(",Asset,Exclusive KiB,Unique Refs KiB,Unique Ref Count,Shared Refs KiB,Shared Ref Count,Total KiB"));
-				}
-				else
-				{
-					OutputDevice->Logf(
-						TEXT(" %100s %20s %20s %15s %20s %15s %25s"),
-						TEXT("Asset"),
-						TEXT("Exclusive KiB"),
-						TEXT("Unique Refs KiB"),
-						TEXT("Unique Ref Count"),
-						TEXT("Shared Refs KiB"),
-						TEXT("Shared Ref Count"),
-						TEXT("Total KiB")
-					);
-				}
+		// Asset dependencies listing
+		if (bShowDependencies)
+		{
+			if (CommonArgs.bCSV)
+			{
+				CurrentOutputDevice->Logf(TEXT(",Asset,Dependency,Unique KiB,Shared KiB"));
+			}
+			else
+			{
+				CurrentOutputDevice->Logf(TEXT(" %100s %100s %20s %20s"),
+					TEXT("Asset"),
+					TEXT("Dependency"),
+					TEXT("Unique KiB"),
+					TEXT("Shared KiB")
+				);
+			}
 
-				// Asset listing
-				for (TPair<FName, FAssetMemoryDetails>& Asset : AssetMemoryMap)
+			for (TPair<FName, FAssetMemoryDetails>& Asset : AssetMemoryMap)
+			{
+				for (TPair<FName, FAssetMemoryBreakdown>& Dep : Asset.Value.DependenciesToMemoryMap)
 				{
-					FAssetMemoryDetails& AssetMemory = Asset.Value;
-					FAssetMemoryBreakdown& AssetMemoryDetails = AssetMemory.MemoryBreakdown;
+					const FString DependencyAssetName = Dep.Key.ToString();
+					const FAssetMemoryBreakdown& DepedencyMemoryDetails = Dep.Value;
 
-					if (bCSV)
+					if (CommonArgs.bCSV)
 					{
-						OutputDevice->Logf(TEXT(",%s,%.2f,%.2f,%d,%.2f,%d,%.2f"),*Asset.Key.ToString(),
-							AssetMemoryDetails.ExclusiveSize / 1024.f, 
-							AssetMemoryDetails.UniqueSize / 1024.f, 
-							AssetMemory.UniqueRefCount, 
-							AssetMemoryDetails.SharedSize / 1024.f,
-							AssetMemory.SharedRefCount,
-							AssetMemoryDetails.TotalSize / 1024.f);
+						CurrentOutputDevice->Logf(TEXT(",%s,%s,%.2f,%.2f"), *Asset.Key.ToString(), *DependencyAssetName,
+							DepedencyMemoryDetails.UniqueSize / 1024.f, DepedencyMemoryDetails.SharedSize / 1024.f);
 					}
 					else
 					{
-						OutputDevice->Logf(
-							TEXT(" %100s %20.2f %20.2f %15d %20.2f %15d %25.2f"),
+						CurrentOutputDevice->Logf(TEXT(" %100s %100s %20.2f %20.2f"),
 							*Asset.Key.ToString(),
-							AssetMemoryDetails.ExclusiveSize / 1024.f, 
-							AssetMemoryDetails.SharedSize / 1024.f,
-							AssetMemory.SharedRefCount,
-							AssetMemoryDetails.UniqueSize / 1024.f, 
-							AssetMemory.UniqueRefCount, 
-							AssetMemoryDetails.TotalSize / 1024.f
-						);
-					}
-				}
-
-				// Asset dependencies listing
-				if (bShowDependencies)
-				{
-					if (bCSV)
-					{
-						OutputDevice->Logf(TEXT(",Asset,Dependency,Unique KiB,Shared KiB"));
-					}
-					else
-					{
-						OutputDevice->Logf(TEXT(" %100s %100s %20s %20s"),
-							TEXT("Asset"),
-							TEXT("Dependency"),
-							TEXT("Unique KiB"),
-							TEXT("Shared KiB")
+							*DependencyAssetName,
+							DepedencyMemoryDetails.UniqueSize / 1024.f,
+							DepedencyMemoryDetails.SharedSize / 1024.f
 						);
 					}
-
-					for (TPair<FName, FAssetMemoryDetails>& Asset : AssetMemoryMap)
-					{
-						for (TPair<FName, FAssetMemoryBreakdown>& Dep : Asset.Value.DependenciesToMemoryMap)
-						{
-							const FString DependencyAssetName = Dep.Key.ToString();
-							const FAssetMemoryBreakdown& DepedencyMemoryDetails = Dep.Value;
-
-							if (bCSV)
-							{
-								OutputDevice->Logf(TEXT(",%s,%s,%.2f,%.2f"), *Asset.Key.ToString(), *DependencyAssetName,
-									DepedencyMemoryDetails.UniqueSize / 1024.f, DepedencyMemoryDetails.SharedSize / 1024.f);
-							}
-							else
-							{
-								OutputDevice->Logf(TEXT(" %100s %100s %20.2f %20.2f"),
-									*Asset.Key.ToString(), 
-									*DependencyAssetName,
-									DepedencyMemoryDetails.UniqueSize / 1024.f, 
-									DepedencyMemoryDetails.SharedSize / 1024.f
-								);
-							}
-						}
-					}
-				}
-
-				if (bCSV)
-				{
-					OutputDevice->Logf(TEXT(",TOTAL KiB,%.2f"), TotalCollectionSize / 1024.f);
-				}
-				else
-				{
-					OutputDevice->Logf(TEXT("TOTAL KiB: %.2f"), TotalCollectionSize / 1024.f);
 				}
-
-				bResult = true;
 			}
 		}
-#endif
 
-		if (FileArWrapper != nullptr)
+		if (CommonArgs.bCSV)
 		{
-			FileArWrapper->TearDown();
+			CurrentOutputDevice->Logf(TEXT(",TOTAL KiB,%.2f"), TotalCollectionSize / 1024.f);
+		}
+		else
+		{
+			CurrentOutputDevice->Logf(TEXT("TOTAL KiB: %.2f"), TotalCollectionSize / 1024.f);
 		}
-			
-		delete FileArWrapper;
-		delete FileAr;
 	}
-	
-	return bResult;
-}
+}));
 
-namespace MemoryUsageQueries
-{
+#endif
 
-void RegisterConsoleAutoCompleteEntries(TArray<FAutoCompleteCommand>& AutoCompleteList)
+namespace MemoryUsageQueries
 {
-	const UConsoleSettings* ConsoleSettings = GetDefault<UConsoleSettings>();
-
-	{
-		FAutoCompleteCommand AutoCompleteCommand;
-		AutoCompleteCommand.Command = TEXT("MemQuery Usage");
-		AutoCompleteCommand.Desc = TEXT("Name=<AssetName> Prints memory usage of the specified asset.");
-		AutoCompleteCommand.Color = ConsoleSettings->AutoCompleteCommandColor;
-		AutoCompleteList.Add(MoveTemp(AutoCompleteCommand));
-	}
-
-	{
-		FAutoCompleteCommand AutoCompleteCommand;
-		AutoCompleteCommand.Command = TEXT("MemQuery CombinedUsage");
-		AutoCompleteCommand.Desc = TEXT("Names=\"<AssetName1> <AssetName2> ...\" Prints combined memory usage of the specified assets (including all dependencies).");
-		AutoCompleteCommand.Color = ConsoleSettings->AutoCompleteCommandColor;
-		AutoCompleteList.Add(MoveTemp(AutoCompleteCommand));
-	}
-
-	{
-		FAutoCompleteCommand AutoCompleteCommand;
-		AutoCompleteCommand.Command = TEXT("MemQuery SharedUsage");
-		AutoCompleteCommand.Desc = TEXT("Names=\"<AssetName1> <AssetName2> ...\" Prints shared memory usage of the specified assets (including only dependencies shared by the specified assets).");
-		AutoCompleteCommand.Color = ConsoleSettings->AutoCompleteCommandColor;
-		AutoCompleteList.Add(MoveTemp(AutoCompleteCommand));
-	}
-
-	{
-		FAutoCompleteCommand AutoCompleteCommand;
-		AutoCompleteCommand.Command = TEXT("MemQuery UniqueUsage");
-		AutoCompleteCommand.Desc = TEXT("Names=\"<AssetName1> <AssetName2> ...\" Prints unique memory usage of the specified assets (including only dependencies unique to the specified assets).");
-		AutoCompleteCommand.Color = ConsoleSettings->AutoCompleteCommandColor;
-		AutoCompleteList.Add(MoveTemp(AutoCompleteCommand));
-	}
-
-	{
-		FAutoCompleteCommand AutoCompleteCommand;
-		AutoCompleteCommand.Command = TEXT("MemQuery CommonUsage");
-		AutoCompleteCommand.Desc = TEXT("Names=\"<AssetName1> <AssetName2> ...\" Prints common memory usage of the specified assets (including only dependencies that are not unique to the specified assets).");
-		AutoCompleteCommand.Color = ConsoleSettings->AutoCompleteCommandColor;
-		AutoCompleteList.Add(MoveTemp(AutoCompleteCommand));
-	}
-
-	{
-		FAutoCompleteCommand AutoCompleteCommand;
-		AutoCompleteCommand.Command = TEXT("MemQuery Dependencies");
-		AutoCompleteCommand.Desc = TEXT("Name=<AssetName> Limit=<n> Lists dependencies of the specified asset, sorted by size.");
-		AutoCompleteCommand.Color = ConsoleSettings->AutoCompleteCommandColor;
-		AutoCompleteList.Add(MoveTemp(AutoCompleteCommand));
-	}
-
-	{
-		FAutoCompleteCommand AutoCompleteCommand;
-		AutoCompleteCommand.Command = TEXT("MemQuery CombinedDependencies");
-		AutoCompleteCommand.Desc = TEXT("Names=\"<AssetName1> <AssetName2> ...\" Limit=<n> Lists n largest dependencies of the specified assets, sorted by size.");
-		AutoCompleteCommand.Color = ConsoleSettings->AutoCompleteCommandColor;
-		AutoCompleteList.Add(MoveTemp(AutoCompleteCommand));
-	}
-
-	{
-		FAutoCompleteCommand AutoCompleteCommand;
-		AutoCompleteCommand.Command = TEXT("MemQuery SharedDependencies");
-		AutoCompleteCommand.Desc = TEXT("Names=\"<AssetName1> <AssetName2> ...\" Limit=<n> Lists n largest dependencies that are shared by the specified assets, sorted by size.");
-		AutoCompleteCommand.Color = ConsoleSettings->AutoCompleteCommandColor;
-		AutoCompleteList.Add(MoveTemp(AutoCompleteCommand));
-	}
-		
-	{
-		FAutoCompleteCommand AutoCompleteCommand;
-		AutoCompleteCommand.Command = TEXT("MemQuery UniqueDependencies");
-		AutoCompleteCommand.Desc = TEXT("Names=\"<AssetName1> <AssetName2> ...\" Limit=<n> Lists n largest dependencies that are unique to the specified assets, sorted by size.");
-		AutoCompleteCommand.Color = ConsoleSettings->AutoCompleteCommandColor;
-		AutoCompleteList.Add(MoveTemp(AutoCompleteCommand));
-	}
-
-	{
-		FAutoCompleteCommand AutoCompleteCommand;
-		AutoCompleteCommand.Command = TEXT("MemQuery CommonDependencies");
-		AutoCompleteCommand.Desc = TEXT("Names=\"<AssetName1> <AssetName2> ...\" Limit=<n> Lists n largest dependencies that are NOT unique to the specified assets, sorted by size.");
-		AutoCompleteCommand.Color = ConsoleSettings->AutoCompleteCommandColor;
-		AutoCompleteList.Add(MoveTemp(AutoCompleteCommand));
-	}
-
-	{
-		FAutoCompleteCommand AutoCompleteCommand;
-		AutoCompleteCommand.Command = TEXT("MemQuery ListAssets");
-		AutoCompleteCommand.Desc = TEXT("Name=<AssetNameSubstring> Group=<GroupName> Class=<ClassName> Limit=<n> Lists n largest assets.");
-		AutoCompleteCommand.Color = ConsoleSettings->AutoCompleteCommandColor;
-		AutoCompleteList.Add(MoveTemp(AutoCompleteCommand));
-	}
-
-	{
-		FAutoCompleteCommand AutoCompleteCommand;
-		AutoCompleteCommand.Command = TEXT("MemQuery ListClasses");
-		AutoCompleteCommand.Desc = TEXT("Group=<GroupName> Asset=<AssetName> Limit=<n> Lists n largest classes.");
-		AutoCompleteCommand.Color = ConsoleSettings->AutoCompleteCommandColor;
-		AutoCompleteList.Add(MoveTemp(AutoCompleteCommand));
-	}
-
-	{
-		FAutoCompleteCommand AutoCompleteCommand;
-		AutoCompleteCommand.Command = TEXT("MemQuery ListGroups");
-		AutoCompleteCommand.Desc = TEXT("Asset=<AssetName> Class=<ClassName> Limit=<n> Lists n largest groups.");
-		AutoCompleteCommand.Color = ConsoleSettings->AutoCompleteCommandColor;
-		AutoCompleteList.Add(MoveTemp(AutoCompleteCommand));
-	}
-
-	const UMemoryUsageQueriesConfig* MemQueryConfig = GetDefault<UMemoryUsageQueriesConfig>();
-
-	for (auto It = MemQueryConfig->SavingsPresets.CreateConstIterator(); It; ++It)
-	{
-		FAutoCompleteCommand AutoCompleteCommand;
-		AutoCompleteCommand.Command = FString::Printf(TEXT("MemQuery Savings %s"), *It.Key());
-		AutoCompleteCommand.Desc = FString::Printf(TEXT("Limit=<n> Lists potential savings among %s. How much memory can be saved it we delete certain object."), *It.Key());
-		AutoCompleteCommand.Color = ConsoleSettings->AutoCompleteCommandColor;
-		AutoCompleteList.Add(MoveTemp(AutoCompleteCommand));
-	}
-
-	for (auto It = MemQueryConfig->Collections.CreateConstIterator(); It; ++It)
-	{
-		FAutoCompleteCommand AutoCompleteCommand;
-		AutoCompleteCommand.Command = FString::Printf(TEXT("MemQuery Collection %s"), *It->Name);
-		AutoCompleteCommand.Desc = TEXT("Lists memory used by a collection. Can show dependency breakdown. [-csv, -showdeps]");
-		AutoCompleteCommand.Color = ConsoleSettings->AutoCompleteCommandColor;
-		AutoCompleteList.Add(MoveTemp(AutoCompleteCommand));
-	}
-}
 
 IMemoryUsageInfoProvider* GetCurrentMemoryUsageInfoProvider()
 {
diff --git a/Engine/Plugins/MemoryUsageQueries/Source/MemoryUsageQueries/Private/MemoryUsageQueriesModule.cpp b/Engine/Plugins/MemoryUsageQueries/Source/MemoryUsageQueries/Private/MemoryUsageQueriesModule.cpp
--- a/Engine/Plugins/MemoryUsageQueries/Source/MemoryUsageQueries/Private/MemoryUsageQueriesModule.cpp
+++ b/Engine/Plugins/MemoryUsageQueries/Source/MemoryUsageQueries/Private/MemoryUsageQueriesModule.cpp
@@ -8,12 +8,10 @@
 
 void FMemoryUsageQueriesModule::StartupModule()
 {
-	Handle_RegisterConsoleAutoCompleteEntries = UConsole::RegisterConsoleAutoCompleteEntries.AddStatic(MemoryUsageQueries::RegisterConsoleAutoCompleteEntries);
 }
 
 void FMemoryUsageQueriesModule::ShutdownModule()
 {
-	UConsole::RegisterConsoleAutoCompleteEntries.Remove(Handle_RegisterConsoleAutoCompleteEntries);
 }
 
 IMPLEMENT_MODULE(FMemoryUsageQueriesModule, MemoryUsageQueries);
diff --git a/Engine/Plugins/MemoryUsageQueries/Source/MemoryUsageQueries/Public/MemoryUsageQueriesModule.h b/Engine/Plugins/MemoryUsageQueries/Source/MemoryUsageQueries/Public/MemoryUsageQueriesModule.h
--- a/Engine/Plugins/MemoryUsageQueries/Source/MemoryUsageQueries/Public/MemoryUsageQueriesModule.h
+++ b/Engine/Plugins/MemoryUsageQueries/Source/MemoryUsageQueries/Public/MemoryUsageQueriesModule.h
@@ -11,9 +11,6 @@ public:
 	/** IModuleInterface implementation */
 	virtual void StartupModule() override;
 	virtual void ShutdownModule() override;
-
-private:
-	FDelegateHandle Handle_RegisterConsoleAutoCompleteEntries;
 };
 
 #if UE_ENABLE_INCLUDE_ORDER_DEPRECATED_IN_5_2
