Fix to compiler error on lambda capture.

diff --git a/Engine/Plugins/MemoryUsageQueries/Source/MemoryUsageQueries/Private/MemoryUsageQueries.cpp b/Engine/Plugins/MemoryUsageQueries/Source/MemoryUsageQueries/Private/MemoryUsageQueries.cpp
--- a/Engine/Plugins/MemoryUsageQueries/Source/MemoryUsageQueries/Private/MemoryUsageQueries.cpp
+++ b/Engine/Plugins/MemoryUsageQueries/Source/MemoryUsageQueries/Private/MemoryUsageQueries.cpp
@@ -574,6 +574,20 @@ FAutoConsoleCommandWithWorldArgsAndOutputDevice GMemQueryListGroups(
 	}
 }));
 
+// Will return true if the Package Name matches any of the conditions in the array of Paths
+static bool PackageNameMatches(const FString& PackageName, const TArray<FString>& Conditions)
+{
+	for (const FString& Condition : Conditions)
+	{
+		if ((FWildcardString::ContainsWildcards(*Condition) && FWildcardString::IsMatch(*Condition, *PackageName)) || (PackageName.Contains(Condition)))
+		{
+			return true;
+		}
+	}
+
+	return false;
+};
+
 FAutoConsoleCommandWithWorldArgsAndOutputDevice GMemQueryCollections(
 	TEXT("MemQuery.Collection"),
 	TEXT("Lists memory used by a collection. Can show dependency breakdown. Pass -showdeps to list dependencies."),
@@ -614,19 +628,6 @@ FAutoConsoleCommandWithWorldArgsAndOutputDevice GMemQueryCollections(
 			break;
 		}
 
-		// Will return true if the Package Name matches any of the conditions in the array of Paths
-		auto PackageNameMatches = [](const FString& PackageName, const TArray<FString>& Conditions)->bool {
-			for (const FString& Condition : Conditions)
-			{
-				if ((FWildcardString::ContainsWildcards(*Condition) && FWildcardString::IsMatch(*Condition, *PackageName)) || (PackageName.Contains(Condition)))
-				{
-					return true;
-				}
-			}
-
-			return false;
-		};
-
 		// See if any of the asset paths match those of our matching paths and are valid
 		TArray<FString> PackageNames;
 
