Removing extra semi-colon from function declaration. Coding standard.

diff --git a/Engine/Plugins/MemoryUsageQueries/Source/MemoryUsageQueries/Private/MemoryUsageQueries.cpp b/Engine/Plugins/MemoryUsageQueries/Source/MemoryUsageQueries/Private/MemoryUsageQueries.cpp
--- a/Engine/Plugins/MemoryUsageQueries/Source/MemoryUsageQueries/Private/MemoryUsageQueries.cpp
+++ b/Engine/Plugins/MemoryUsageQueries/Source/MemoryUsageQueries/Private/MemoryUsageQueries.cpp
@@ -586,7 +586,7 @@ static bool PackageNameMatches(const FString& PackageName, const TArray<FString>
 	}
 
 	return false;
-};
+}
 
 FAutoConsoleCommandWithWorldArgsAndOutputDevice GMemQueryCollections(
 	TEXT("MemQuery.Collection"),
