Removed superfluous forward declaration

#rnx
#prefliight none

diff --git a/Engine/Plugins/MemoryUsageQueries/Source/MemoryUsageQueries/Private/MemoryUsageQueriesPrivate.h b/Engine/Plugins/MemoryUsageQueries/Source/MemoryUsageQueries/Private/MemoryUsageQueriesPrivate.h
--- a/Engine/Plugins/MemoryUsageQueries/Source/MemoryUsageQueries/Private/MemoryUsageQueriesPrivate.h
+++ b/Engine/Plugins/MemoryUsageQueries/Source/MemoryUsageQueries/Private/MemoryUsageQueriesPrivate.h
@@ -4,7 +4,6 @@
 #include "UObject/FastReferenceCollector.h"
 
 class IMemoryUsageInfoProvider;
-enum class EGCTokenType : uint32;
 
 namespace MemoryUsageQueries::Internal
 {
