[MemQuery] Add functionality to list assets by category. Used for tracking UntaggedAssets by LLM tag.

Example usage: MemQuery.ListAssetsCategorized Asset=UntaggedAsset -notrunc

Additionally, ~40% perf improvement to FLLMTracker::GetTagsNamesWithAmountFiltered by changing it to an array.

#tests Tested existing commands remain unchanged and the new command works as expected. Both existing and new commands were tested with/without -llm -llmtagsets=assets,assetclasses and LLM_ALLOW_ASSETS_TAGS.

diff --git a/Engine/Plugins/MemoryUsageQueries/Source/MemoryUsageQueries/Private/MemoryUsageQueries.cpp b/Engine/Plugins/MemoryUsageQueries/Source/MemoryUsageQueries/Private/MemoryUsageQueries.cpp
--- a/Engine/Plugins/MemoryUsageQueries/Source/MemoryUsageQueries/Private/MemoryUsageQueries.cpp
+++ b/Engine/Plugins/MemoryUsageQueries/Source/MemoryUsageQueries/Private/MemoryUsageQueries.cpp
@@ -142,6 +142,7 @@ struct FCommonParameters
 	
 	FName Group = NAME_None;
 	FName Class = NAME_None;
+	FName Category = NAME_None;
 
 	FString Name;
 	FString Names;
@@ -170,6 +171,12 @@ struct FCommonParameters
 		{
 			Class = FName(*ClassName);
 		}
+
+		FString CategoryName;
+		if (FParse::Value(Args, TEXT("Category="), CategoryName))
+		{
+			Category = FName(*CategoryName);
+		}
 	}
 
 	FCommonParameters() = delete;
@@ -504,7 +511,7 @@ FAutoConsoleCommandWithWorldArgsAndOutputDevice GMemQuerySavings(
 
 FAutoConsoleCommandWithWorldArgsAndOutputDevice GMemQueryListAssets(
 	TEXT("MemQuery.ListAssets"),
-	TEXT("Name=<AssetNameSubstring> Group=<GroupName> Class=<ClassName> Limit=<n> Lists n largest assets."),
+	TEXT("Asset=<AssetNameSubstring> Group=<GroupName> Class=<ClassName> Limit=<n> Lists n largest assets."),
 	FConsoleCommandWithWorldArgsAndOutputDeviceDelegate::CreateStatic(
 		[](const TArray<FString>& Args, UWorld* World, FOutputDevice& Ar)
 {
@@ -534,6 +541,29 @@ FAutoConsoleCommandWithWorldArgsAndOutputDevice GMemQueryListAssets(
 	}
 }));
 
+FAutoConsoleCommandWithWorldArgsAndOutputDevice GMemQueryListAssetsCategorized(
+	TEXT("MemQuery.ListAssetsCategorized"),
+	TEXT("Asset=<AssetNameSubstring> Group=<GroupName> Class=<ClassName> Category=<CategoryName(None,Assets,AssetClasses)> Limit=<n> Lists n largest assets categorized by Category."),
+	FConsoleCommandWithWorldArgsAndOutputDeviceDelegate::CreateStatic(
+		[](const TArray<FString>& Args, UWorld* World, FOutputDevice& Ar)
+{
+	const FString Cmd = FString::Join(Args, TEXT(" "));
+
+	FScopedOutputDevice ScopedOutputDevice(&Ar);
+	FCommonParameters CommonArgs(*Cmd);
+
+	ScopedOutputDevice.OpenLogFile(CommonArgs.LogFileName, CommonArgs.bCSV);
+
+	bool bSuccess;
+	TMap<FName, uint64> AssetsWithSize;						
+	bSuccess = MemoryUsageQueries::GetFilteredPackagesCategorizedWithSize(AssetsWithSize, CommonArgs.Group, CommonArgs.AssetName, CommonArgs.Class, CommonArgs.Category, &ScopedOutputDevice.GetOutputDevice());
+
+	if (bSuccess)
+	{
+		MemoryUsageQueries::Internal::PrintTagsWithSize(ScopedOutputDevice.GetOutputDevice(), AssetsWithSize, TEXT("largest assets"), CommonArgs.bTruncate, CommonArgs.Limit, CommonArgs.bCSV);
+	}
+}));
+
 FAutoConsoleCommandWithWorldArgsAndOutputDevice GMemQueryListClasses(
 	TEXT("MemQuery.ListClasses"),
 	TEXT("Group=<GroupName> Asset=<AssetName> Limit=<n> Lists n largest classes."),
@@ -1018,6 +1048,60 @@ bool GetFilteredPackagesWithSize(TMap<FName, uint64>& OutPackagesWithSize, FName
 	return true;
 }
 
+bool GetFilteredPackagesCategorizedWithSize(TMap<FName, uint64>& OutPackagesWithSize, FName GroupName /*= NAME_None*/, FString AssetSubstring /*= FString()*/, FName ClassName /*= NAME_None*/, FName CategoryName /*= NAME_None*/, FOutputDevice* ErrorOutput /*= GLog*/)
+{
+	TArray<FLLMTagSetAllocationFilter> Filters;
+
+	FName LongName = NAME_None;
+	if (!AssetSubstring.IsEmpty())
+	{
+		// Special handling for "Untagged" assets as llm uses this tag to bucket allocations without asset scopes.
+		FName UntaggedAssetTagName = LLMGetUntaggedTagName(ELLMTagSet::Assets);
+		if (FName(AssetSubstring) == UntaggedAssetTagName)
+		{
+			LongName = UntaggedAssetTagName;
+		}
+		else
+		{
+			if (!Internal::GetLongName(AssetSubstring, LongName, ErrorOutput))
+			{
+				return false;
+			}
+		}
+	}
+
+	if (LongName != NAME_None)
+	{
+		Filters.Add({ LongName, ELLMTagSet::Assets });
+	}
+
+	if (GroupName != NAME_None)
+	{
+		Filters.Add({ GroupName, ELLMTagSet::None });
+	}
+
+	if (ClassName != NAME_None)
+	{
+		Filters.Add({ ClassName, ELLMTagSet::AssetClasses });
+	}
+
+	ELLMTagSet TagSetCategory = ELLMTagSet::None;
+	if (CategoryName == FName("Assets"))
+	{
+		TagSetCategory = ELLMTagSet::Assets;
+	}
+	else  if (CategoryName == FName("AssetClasses"))
+	{
+		TagSetCategory = ELLMTagSet::AssetClasses;
+	}
+
+	MemoryUsageInfoProviderLLM.GetFilteredTagsWithSize(OutPackagesWithSize, ELLMTracker::Default, TagSetCategory, Filters, ErrorOutput);
+
+	OutPackagesWithSize.ValueSort(TGreater<uint64>());
+
+	return true;
+}
+
 bool GetFilteredClassesWithSize(TMap<FName, uint64>& OutClassesWithSize, FName GroupName /* = NAME_None */, FString AssetName /* = FString() */, FOutputDevice* ErrorOutput /* = GLog */)
 {
 	TArray<FLLMTagSetAllocationFilter> Filters;
diff --git a/Engine/Plugins/MemoryUsageQueries/Source/MemoryUsageQueries/Public/MemoryUsageQueries.h b/Engine/Plugins/MemoryUsageQueries/Source/MemoryUsageQueries/Public/MemoryUsageQueries.h
--- a/Engine/Plugins/MemoryUsageQueries/Source/MemoryUsageQueries/Public/MemoryUsageQueries.h
+++ b/Engine/Plugins/MemoryUsageQueries/Source/MemoryUsageQueries/Public/MemoryUsageQueries.h
@@ -30,6 +30,7 @@ MEMORYUSAGEQUERIES_API bool GetDependenciesWithSizeCommon(const IMemoryUsageInfo
 
 #if ENABLE_LOW_LEVEL_MEM_TRACKER
 MEMORYUSAGEQUERIES_API bool GetFilteredPackagesWithSize(TMap<FName, uint64>& OutPackagesWithSize, FName GroupName = NAME_None, FString AssetSubstring = FString(), FName ClassName = NAME_None, FOutputDevice* ErrorOutput = GLog);
+MEMORYUSAGEQUERIES_API bool GetFilteredPackagesCategorizedWithSize(TMap<FName, uint64>& OutPackagesWithSize, FName GroupName = NAME_None, FString AssetSubstring = FString(), FName ClassName = NAME_None, FName CategoryName = NAME_None, FOutputDevice* ErrorOutput = GLog);
 MEMORYUSAGEQUERIES_API bool GetFilteredClassesWithSize(TMap<FName, uint64>& OutClassesWithSize, FName GroupName = NAME_None, FString AssetName = FString(), FOutputDevice* ErrorOutput = GLog);
 MEMORYUSAGEQUERIES_API bool GetFilteredGroupsWithSize(TMap<FName, uint64>& OutGroupsWithSize, FString AssetName = FString(), FName ClassName = NAME_None, FOutputDevice* ErrorOutput = GLog);
 MEMORYUSAGEQUERIES_API bool GatherDependenciesForPackages(const IMemoryUsageInfoProvider* MemoryUsageInfoProvider, const TArray<FString>& PackageNames,
diff --git a/Engine/Source/Runtime/Core/Private/HAL/LowLevelMemTracker.cpp b/Engine/Source/Runtime/Core/Private/HAL/LowLevelMemTracker.cpp
--- a/Engine/Source/Runtime/Core/Private/HAL/LowLevelMemTracker.cpp
+++ b/Engine/Source/Runtime/Core/Private/HAL/LowLevelMemTracker.cpp
@@ -1709,6 +1709,29 @@ extern const TCHAR* LLMGetTagName(ELLMTag Tag)
 	}
 }
 
+extern const FName LLMGetUntaggedTagName(ELLMTagSet TagSet)
+{
+	switch (TagSet)
+	{
+		case ELLMTagSet::None:
+		{
+			return TagName_Untagged;
+		}
+		case ELLMTagSet::Assets:
+		{
+			return TagName_UntaggedAsset;
+		}
+		case ELLMTagSet::AssetClasses:
+		{
+			return TagName_UntaggedAssetClass;
+		}
+		default:
+		{
+			return NAME_None;
+		}
+	}
+}
+
 extern const ANSICHAR* LLMGetTagNameANSI(ELLMTag Tag)
 {
 #define LLM_TAG_NAME_ARRAY(Enum,Str,Stat,Group,ParentTag) Str,
@@ -5612,27 +5635,23 @@ void FLLMTracker::GetTagsNamesWithAmountFiltered(TMap<FName, uint64>& OutTagsNam
 		FName Name;
 		uint64 Size = 0;
 	};
-#if LLM_ENABLED_FULL_TAGS
-	TMap<int32, FLocalTagData> CompressedTagToTagData;
-#else
-	TMap<ELLMTag, FLocalTagData> CompressedTagToTagData;
-#endif
-	int32 NumTags;
-	{
-		FReadScopeLock TagDataScopeLock(LLMRef.TagDataLock);
-		NumTags = LLMRef.TagDatas->Num();
-	}
-	CompressedTagToTagData.Reserve(NumTags);
+	TArray<FLocalTagData> TagDataCached;
 	{
 		FReadScopeLock TagDataScopeLock(LLMRef.TagDataLock);
+		const int32 NumTags = LLMRef.TagDatas->Num();
+		TagDataCached.Init(FLocalTagData{}, NumTags);
 		for (FTagData* TagData : (*LLMRef.TagDatas))
 		{
 			FLLMTracker::FLowLevelAllocInfo AllocInfo;
 			AllocInfo.SetTag(TagData, LLMRef);
-			FLocalTagData& Data = CompressedTagToTagData.FindOrAdd(AllocInfo.GetCompressedTag());
-			if (Data.Name == NAME_None)
+			int32 CompressedTagIndex = static_cast<int32>(AllocInfo.GetCompressedTag());
+			if(ensure(TagDataCached.IsValidIndex(CompressedTagIndex)))
 			{
-				Data.Name = TagData->GetName();
+				FLocalTagData& Data = TagDataCached[CompressedTagIndex];
+				if (Data.Name == NAME_None)
+				{
+					Data.Name = TagData->GetName();
+				}
 			}
 		}
 	}
@@ -5645,13 +5664,21 @@ void FLLMTracker::GetTagsNamesWithAmountFiltered(TMap<FName, uint64>& OutTagsNam
 		{
 			FLocalTagData* Data = nullptr;
 #if LLM_ALLOW_ASSETS_TAGS
-			Data = CompressedTagToTagData.Find(Tuple.Value2.GetCompressedTag(Filter.TagSet));
+			int32 CompressedTagIndex = static_cast<int32>(Tuple.Value2.GetCompressedTag(Filter.TagSet));
 #else
-			if (Filter.TagSet == ELLMTagSet::None)
+			int32 CompressedTagIndex = static_cast<int32>(Tuple.Value2.GetCompressedTag());
+#endif
+			if(TagDataCached.IsValidIndex(CompressedTagIndex))
 			{
-				Data = CompressedTagToTagData.Find(Tuple.Value2.GetCompressedTag());
-			}
+#if LLM_ALLOW_ASSETS_TAGS
+				Data = &TagDataCached[CompressedTagIndex];
+#else
+				if (Filter.TagSet == ELLMTagSet::None)
+				{
+					Data = &TagDataCached[CompressedTagIndex];
+				}
 #endif
+			}
 			if (!Data || Data->Name != Filter.Name)
 			{
 				bIncludeAllocation = false;
@@ -5662,10 +5689,11 @@ void FLLMTracker::GetTagsNamesWithAmountFiltered(TMap<FName, uint64>& OutTagsNam
 		if (bIncludeAllocation)
 		{
 #if LLM_ALLOW_ASSETS_TAGS
-			FLocalTagData* Data = CompressedTagToTagData.Find(Tuple.Value2.GetCompressedTag(TagSet));
+			int32 CompressedTagIndex = Tuple.Value2.GetCompressedTag(TagSet);
 #else
-			FLocalTagData* Data = CompressedTagToTagData.Find(Tuple.Value2.GetCompressedTag());
+			int32 CompressedTagIndex = static_cast<int32>(Tuple.Value2.GetCompressedTag());
 #endif
+			FLocalTagData* Data = TagDataCached.IsValidIndex(CompressedTagIndex) ? &TagDataCached[CompressedTagIndex] : nullptr;
 			if (Data)
 			{
 				Data->Size += Tuple.Value1;
@@ -5674,16 +5702,12 @@ void FLLMTracker::GetTagsNamesWithAmountFiltered(TMap<FName, uint64>& OutTagsNam
 	}
 	AllocationMap.UnlockAll();
 #endif // !UE_ONLY_USE_PLATFORM_TRACKER
-
-#if LLM_ENABLED_FULL_TAGS
-	for (TPair<int32, FLocalTagData>& Pair : CompressedTagToTagData)
-#else
-	for (TPair<ELLMTag, FLocalTagData>& Pair : CompressedTagToTagData)
-#endif
+	
+	for (const FLocalTagData& LocalTagData : TagDataCached)
 	{
-		if (Pair.Value.Size != 0)
+		if (LocalTagData.Size != 0)
 		{
-			OutTagsNamesWithAmount.FindOrAdd(Pair.Value.Name, 0) += Pair.Value.Size;
+			OutTagsNamesWithAmount.FindOrAdd(LocalTagData.Name, 0) += LocalTagData.Size;
 		}
 	}
 }
diff --git a/Engine/Source/Runtime/Core/Public/HAL/LowLevelMemTracker.h b/Engine/Source/Runtime/Core/Public/HAL/LowLevelMemTracker.h
--- a/Engine/Source/Runtime/Core/Public/HAL/LowLevelMemTracker.h
+++ b/Engine/Source/Runtime/Core/Public/HAL/LowLevelMemTracker.h
@@ -315,6 +315,7 @@ enum class ELLMAllocType
 
 extern const ANSICHAR* LLMGetTagNameANSI(ELLMTag Tag);
 extern const TCHAR* LLMGetTagName(ELLMTag Tag);
+extern CORE_API const FName LLMGetUntaggedTagName(ELLMTagSet TagSet);
 UE_DEPRECATED(4.27, "This function was an unused implementation detail; contact Epic if you need to keep its functionality.")
 extern FName LLMGetTagStatGroup(ELLMTag Tag);
 UE_DEPRECATED(4.27, "This function was an unused implementation detail; contact Epic if you need to keep its functionality.")
