Adding additional function to gather assets and their dependencies with associated sizes.

diff --git a/Engine/Plugins/MemoryUsageQueries/Source/MemoryUsageQueries/Private/MemoryUsageQueries.cpp b/Engine/Plugins/MemoryUsageQueries/Source/MemoryUsageQueries/Private/MemoryUsageQueries.cpp
--- a/Engine/Plugins/MemoryUsageQueries/Source/MemoryUsageQueries/Private/MemoryUsageQueries.cpp
+++ b/Engine/Plugins/MemoryUsageQueries/Source/MemoryUsageQueries/Private/MemoryUsageQueries.cpp
@@ -690,7 +690,7 @@ FAutoConsoleCommandWithWorldArgsAndOutputDevice GMemQueryCollections(
 		// External are dependencies that have additional references outside the set of packages passed in.
 		TMap<FName, uint64> InternalDependencies;
 		TMap<FName, uint64> ExternalDependencies;
-		if (!MemoryUsageQueries::GatherDependenciesForPackages(MemoryUsageQueries::CurrentMemoryUsageInfoProvider, PackageNames, InternalDependencies, ExternalDependencies, &Ar))
+		if (!MemoryUsageQueries::GatherDependenciesForPackages(MemoryUsageQueries::CurrentMemoryUsageInfoProvider, PackageNames, InternalDependencies, ExternalDependencies, MemoryUsageQueries::EDependencyType::EDep_All, &Ar))
 		{
 			Ar.Logf(TEXT("Failed to gather memory usage for dependencies in Collection %s"), *CollectionInfo.Name);
 			break;
@@ -928,22 +928,32 @@ bool GetMemoryUsageCommon(const IMemoryUsageInfoProvider* MemoryUsageInfoProvide
 }
 
 bool GatherDependenciesForPackages(const IMemoryUsageInfoProvider* MemoryUsageInfoProvider, const TArray<FString>& PackageNames,
-	TMap<FName, uint64>& OutInternalDeps, TMap<FName, uint64>& OutExternalDeps, FOutputDevice* ErrorOutput)
+	TMap<FName, uint64>& OutInternalDeps, TMap<FName, uint64>& OutExternalDeps, MemoryUsageQueries::EDependencyType DependencyType, FOutputDevice* ErrorOutput)
 {
+	const bool bFindRemovable = (DependencyType == MemoryUsageQueries::EDependencyType::EDep_All || DependencyType == MemoryUsageQueries::EDependencyType::EDep_Removable);
+	const bool bFindNonRemovable = (DependencyType == MemoryUsageQueries::EDependencyType::EDep_All || DependencyType == MemoryUsageQueries::EDependencyType::EDep_NonRemovable);
+
 	TSet<FName> RemovablePackages;
-	if (!Internal::GetRemovablePackages(PackageNames, RemovablePackages, ErrorOutput))
+	if (bFindRemovable && !Internal::GetRemovablePackages(PackageNames, RemovablePackages, ErrorOutput))
 	{
 		return false;
 	}
 
 	TSet<FName> UnremovablePackages;
-	if (!Internal::GetUnremovablePackages(PackageNames, UnremovablePackages, ErrorOutput))
+	if (bFindNonRemovable && !Internal::GetUnremovablePackages(PackageNames, UnremovablePackages, ErrorOutput))
 	{
 		return false;
 	}
 
-	MemoryUsageInfoProvider->GetAssetsMemoryUsageWithSize(RemovablePackages, OutInternalDeps, ErrorOutput);
-	MemoryUsageInfoProvider->GetAssetsMemoryUsageWithSize(UnremovablePackages, OutExternalDeps, ErrorOutput);
+	if (bFindRemovable)
+	{
+		MemoryUsageInfoProvider->GetAssetsMemoryUsageWithSize(RemovablePackages, OutInternalDeps, ErrorOutput);
+	}
+
+	if (bFindNonRemovable)
+	{
+		MemoryUsageInfoProvider->GetAssetsMemoryUsageWithSize(UnremovablePackages, OutExternalDeps, ErrorOutput);
+	}
 
 	return true;
 }
diff --git a/Engine/Plugins/MemoryUsageQueries/Source/MemoryUsageQueries/Public/MemoryUsageQueries.h b/Engine/Plugins/MemoryUsageQueries/Source/MemoryUsageQueries/Public/MemoryUsageQueries.h
--- a/Engine/Plugins/MemoryUsageQueries/Source/MemoryUsageQueries/Public/MemoryUsageQueries.h
+++ b/Engine/Plugins/MemoryUsageQueries/Source/MemoryUsageQueries/Public/MemoryUsageQueries.h
@@ -15,6 +15,13 @@ extern MEMORYUSAGEQUERIES_API FMemoryUsageInfoProviderLLM MemoryUsageInfoProvide
 
 void RegisterConsoleAutoCompleteEntries(TArray<FAutoCompleteCommand>& AutoCompleteList);
 
+enum EDependencyType
+{
+	EDep_All,
+	EDep_Removable,
+	EDep_NonRemovable
+};
+
 MEMORYUSAGEQUERIES_API IMemoryUsageInfoProvider* GetCurrentMemoryUsageInfoProvider();
 MEMORYUSAGEQUERIES_API bool GetMemoryUsage(const IMemoryUsageInfoProvider* MemoryUsageInfoProvider, const FString& PackageName, uint64& OutExclusiveSize, uint64& OutInclusiveSize, FOutputDevice* ErrorOutput = GLog);
 MEMORYUSAGEQUERIES_API bool GetMemoryUsageCombined(const IMemoryUsageInfoProvider* MemoryUsageInfoProvider, const TArray<FString>& PackageNames, uint64& OutTotalSize, FOutputDevice* ErrorOutput = GLog);
@@ -29,12 +36,13 @@ MEMORYUSAGEQUERIES_API bool GetDependenciesWithSizeUnique(const IMemoryUsageInfo
 MEMORYUSAGEQUERIES_API bool GetDependenciesWithSizeCommon(const IMemoryUsageInfoProvider* MemoryUsageInfoProvider, const TArray<FString>& PackageNames, TMap<FName, uint64>& OutDependenciesWithSize, FOutputDevice* ErrorOutput = GLog);
 
 #if ENABLE_LOW_LEVEL_MEM_TRACKER
+
 MEMORYUSAGEQUERIES_API bool GetFilteredPackagesWithSize(TMap<FName, uint64>& OutPackagesWithSize, FName GroupName = NAME_None, FString AssetSubstring = FString(), FName ClassName = NAME_None, FOutputDevice* ErrorOutput = GLog);
 MEMORYUSAGEQUERIES_API bool GetFilteredPackagesCategorizedWithSize(TMap<FName, uint64>& OutPackagesWithSize, FName GroupName = NAME_None, FString AssetSubstring = FString(), FName ClassName = NAME_None, FName CategoryName = NAME_None, FOutputDevice* ErrorOutput = GLog);
 MEMORYUSAGEQUERIES_API bool GetFilteredClassesWithSize(TMap<FName, uint64>& OutClassesWithSize, FName GroupName = NAME_None, FString AssetName = FString(), FOutputDevice* ErrorOutput = GLog);
 MEMORYUSAGEQUERIES_API bool GetFilteredGroupsWithSize(TMap<FName, uint64>& OutGroupsWithSize, FString AssetName = FString(), FName ClassName = NAME_None, FOutputDevice* ErrorOutput = GLog);
 MEMORYUSAGEQUERIES_API bool GatherDependenciesForPackages(const IMemoryUsageInfoProvider* MemoryUsageInfoProvider, const TArray<FString>& PackageNames,
-	TMap<FName, uint64>& OutInternalDeps, TMap<FName, uint64>& OutExternalDeps, FOutputDevice* ErrorOutput = GLog);
+	TMap<FName, uint64>& OutInternalDeps, TMap<FName, uint64>& OutExternalDeps, MemoryUsageQueries::EDependencyType DependencyType = MemoryUsageQueries::EDependencyType::EDep_All, FOutputDevice* ErrorOutput = GLog);
 #endif
 
 }
