[MemQuery] Refactor code to get rid of FString operations, input handling and provide cleaner API

[RN] Refactor MemQuery API to accept FNames as a package name and excluded user input parsing.

diff --git a/Engine/Plugins/MemoryUsageQueries/Source/MemoryUsageQueries/Private/MemoryUsageInfoProvider.cpp b/Engine/Plugins/MemoryUsageQueries/Source/MemoryUsageQueries/Private/MemoryUsageInfoProvider.cpp
--- a/Engine/Plugins/MemoryUsageQueries/Source/MemoryUsageQueries/Private/MemoryUsageInfoProvider.cpp
+++ b/Engine/Plugins/MemoryUsageQueries/Source/MemoryUsageQueries/Private/MemoryUsageInfoProvider.cpp
@@ -2,76 +2,22 @@
 
 #include "MemoryUsageInfoProvider.h"
 
-uint64 FMemoryUsageInfoProviderLLM::GetAssetMemoryUsage(FName Asset, FOutputDevice* ErrorOutput /* = GLog */) const
+uint64 IMemoryUsageInfoProvider::GetAssetMemoryUsage(FName Asset, FOutputDevice* ErrorOutput) const
 {
-#if ENABLE_LOW_LEVEL_MEM_TRACKER
-	if (FLowLevelMemTracker::IsEnabled())
-	{
-		return FLowLevelMemTracker::Get().GetTagAmountForTracker(ELLMTracker::Default, Asset, ELLMTagSet::Assets, UE::LLM::ESizeParams::Default);
-	}
-#endif
-	ErrorOutput->Logf(TEXT("MemoryUsageInfoProvider Error: LLM is disabled. Please run with -LLM"));
+	ErrorOutput->Logf(TEXT("MemoryUsageInfoProvider Error: No provider is available."));
 	return 0U;
 }
-
-uint64 FMemoryUsageInfoProviderLLM::GetAssetsMemoryUsage(const TSet<FName>& Assets, FOutputDevice* ErrorOutput /* = GLog */) const
+uint64 IMemoryUsageInfoProvider::GetAssetsMemoryUsage(const TSet<FName>& Assets, FOutputDevice* ErrorOutput) const
 {
-#if ENABLE_LOW_LEVEL_MEM_TRACKER
-	if (FLowLevelMemTracker::IsEnabled())
-	{
-		uint64 TotalSize = 0;
-		for (const FName& Asset : Assets)
-		{
-			TotalSize += FLowLevelMemTracker::Get().GetTagAmountForTracker(ELLMTracker::Default, Asset, ELLMTagSet::Assets, UE::LLM::ESizeParams::Default);
-		}
-		return TotalSize;
-	}
-#endif
-	ErrorOutput->Logf(TEXT("MemoryUsageInfoProvider Error: LLM is disabled. Please run with -LLM"));
+	ErrorOutput->Logf(TEXT("MemoryUsageInfoProvider Error: No provider is available."));
 	return 0U;
 }
-
-uint64 FMemoryUsageInfoProviderLLM::GetAssetsMemoryUsageWithSize(const TSet<FName>& Assets, TMap<FName, uint64>& OutSizes, FOutputDevice* ErrorOutput /* = GLog */) const
+uint64 IMemoryUsageInfoProvider::GetAssetsMemoryUsageWithSize(const TSet<FName>& Assets, TMap<FName, uint64>& OutSizes, FOutputDevice* ErrorOutput) const
 {
-#if ENABLE_LOW_LEVEL_MEM_TRACKER
-	if (FLowLevelMemTracker::IsEnabled())
-	{
-		uint64 TotalSize = 0;
-		for (const FName& Asset : Assets)
-		{
-			int64 AssetSize = FLowLevelMemTracker::Get().GetTagAmountForTracker(ELLMTracker::Default, Asset, ELLMTagSet::Assets, UE::LLM::ESizeParams::Default);
-			OutSizes.Add(Asset, AssetSize);
-			TotalSize += AssetSize;
-		}
-		return TotalSize;
-	}
-#endif
-	ErrorOutput->Logf(TEXT("MemoryUsageInfoProvider Error: LLM is disabled. Please run with -LLM"));
+	ErrorOutput->Logf(TEXT("MemoryUsageInfoProvider Error: No provider is available."));
 	return 0U;
 }
-
-void FMemoryUsageInfoProviderLLM::GetAllAssetsWithSize(TMap<FName, uint64>& OutAssets, FOutputDevice* ErrorOutput /* = GLog */) const
-{
-	OutAssets.Reset();
-#if ENABLE_LOW_LEVEL_MEM_TRACKER
-	if (FLowLevelMemTracker::IsEnabled())
-	{
-		FLowLevelMemTracker::Get().GetTrackedTagsNamesWithAmount(OutAssets, ELLMTracker::Default, ELLMTagSet::Assets);
-		return;
-	}
-#endif
-	ErrorOutput->Logf(TEXT("MemoryUsageInfoProvider Error: LLM is disabled. Please run with -LLM"));
-}
-
-#if ENABLE_LOW_LEVEL_MEM_TRACKER
-void FMemoryUsageInfoProviderLLM::GetFilteredTagsWithSize(TMap<FName, uint64>& OutTags, ELLMTracker Tracker, ELLMTagSet TagSet, TArray<FLLMTagSetAllocationFilter>& Filters, FOutputDevice* ErrorOutput /* = GLog */) const
+void IMemoryUsageInfoProvider::GetAllAssetsWithSize(TMap<FName, uint64>& OutAssets, FOutputDevice* ErrorOutput) const
 {
-	OutTags.Reset();
-	if (FLowLevelMemTracker::IsEnabled())
-	{
-		FLowLevelMemTracker::Get().GetTrackedTagsNamesWithAmountFiltered(OutTags, Tracker, TagSet, Filters);
-		return;
-	}
-	ErrorOutput->Logf(TEXT("MemoryUsageInfoProvider Error: LLM is disabled. Please run with -LLM"));
+	ErrorOutput->Logf(TEXT("MemoryUsageInfoProvider Error: No provider is available."));
 }
-#endif
diff --git a/Engine/Plugins/MemoryUsageQueries/Source/MemoryUsageQueries/Private/MemoryUsageInfoProviderLLM.cpp b/Engine/Plugins/MemoryUsageQueries/Source/MemoryUsageQueries/Private/MemoryUsageInfoProviderLLM.cpp
new file mode 100644
--- /dev/null
+++ b/Engine/Plugins/MemoryUsageQueries/Source/MemoryUsageQueries/Private/MemoryUsageInfoProviderLLM.cpp
@@ -0,0 +1,233 @@
+// Copyright Epic Games, Inc. All Rights Reserved.
+
+#include "MemoryUsageInfoProviderLLM.h"
+#include "MemoryUsageQueries.h"
+#include "MemoryUsageQueriesPrivate.h"
+
+#if ENABLE_LOW_LEVEL_MEM_TRACKER
+
+bool FMemoryUsageInfoProviderLLM::IsProviderAvailable() const
+{
+	return FLowLevelMemTracker::IsEnabled();
+}
+
+uint64 FMemoryUsageInfoProviderLLM::GetAssetMemoryUsage(FName Asset) const
+{
+	if (FLowLevelMemTracker::IsEnabled())
+	{
+		return FLowLevelMemTracker::Get().GetTagAmountForTracker(ELLMTracker::Default, Asset, ELLMTagSet::Assets, UE::LLM::ESizeParams::Default);
+	}
+	return 0U;
+}
+
+uint64 FMemoryUsageInfoProviderLLM::GetAssetsMemoryUsage(const TSet<FName>& Assets) const
+{
+	if (FLowLevelMemTracker::IsEnabled())
+	{
+		uint64 TotalSize = 0;
+		for (const FName& Asset : Assets)
+		{
+			TotalSize += FLowLevelMemTracker::Get().GetTagAmountForTracker(ELLMTracker::Default, Asset, ELLMTagSet::Assets, UE::LLM::ESizeParams::Default);
+		}
+		return TotalSize;
+	}
+	return 0U;
+}
+
+uint64 FMemoryUsageInfoProviderLLM::GetAssetsMemoryUsageWithSize(const TSet<FName>& Assets, TMap<FName, uint64>& OutSizes) const
+{
+	if (FLowLevelMemTracker::IsEnabled())
+	{
+		uint64 TotalSize = 0;
+		for (const FName& Asset : Assets)
+		{
+			int64 AssetSize = FLowLevelMemTracker::Get().GetTagAmountForTracker(ELLMTracker::Default, Asset, ELLMTagSet::Assets, UE::LLM::ESizeParams::Default);
+			OutSizes.Add(Asset, AssetSize);
+			TotalSize += AssetSize;
+		}
+		return TotalSize;
+	}
+	return 0U;
+}
+
+void FMemoryUsageInfoProviderLLM::GetAllAssetsWithSize(TMap<FName, uint64>& OutAssets) const
+{
+	OutAssets.Reset();
+	if (FLowLevelMemTracker::IsEnabled())
+	{
+		FLowLevelMemTracker::Get().GetTrackedTagsNamesWithAmount(OutAssets, ELLMTracker::Default, ELLMTagSet::Assets);
+	}
+}
+
+void FMemoryUsageInfoProviderLLM::GetFilteredTagsWithSize(TMap<FName, uint64>& OutTags, ELLMTracker Tracker, ELLMTagSet TagSet, TArray<FLLMTagSetAllocationFilter>& Filters) const
+{
+	OutTags.Reset();
+	if (FLowLevelMemTracker::IsEnabled())
+	{
+		FLowLevelMemTracker::Get().GetTrackedTagsNamesWithAmountFiltered(OutTags, Tracker, TagSet, Filters);
+		return;
+	}
+}
+
+namespace MemoryUsageQueries
+{
+
+namespace Internal
+{
+	FMemoryUsageInfoProviderLLM GMemoryUsageInfoProviderLLM;
+}
+
+IMemoryUsageInfoProvider* GetCurrentMemoryUsageInfoProvider()
+{
+	return &Internal::GMemoryUsageInfoProviderLLM;
+}
+
+bool GetFilteredPackagesWithSize(TMap<FName, uint64>& OutPackagesWithSize, FName GroupName /* = NAME_None */, FString AssetSubstring /* = FString() */, FName ClassName /* = NAME_None */, FOutputDevice* ErrorOutput /*= GLog*/)
+{
+	TArray<FLLMTagSetAllocationFilter> Filters;
+
+	FName LongName = NAME_None;
+	if (!AssetSubstring.IsEmpty())
+	{
+		Internal::GetLongName(AssetSubstring, LongName, ErrorOutput);
+	}
+
+	if (LongName != NAME_None)
+	{
+		Filters.Add({ LongName, ELLMTagSet::Assets });
+	}
+
+	if (GroupName != NAME_None)
+	{
+		Filters.Add({GroupName, ELLMTagSet::None});
+	}
+
+	if (ClassName != NAME_None)
+	{
+		Filters.Add({ClassName, ELLMTagSet::AssetClasses});
+	}
+
+	Internal::GMemoryUsageInfoProviderLLM.GetFilteredTagsWithSize(OutPackagesWithSize, ELLMTracker::Default, ELLMTagSet::Assets, Filters);
+
+	if (AssetSubstring.Len() > 0)
+	{
+		Internal::RemoveFilteredPackages(OutPackagesWithSize, AssetSubstring);
+	}
+
+	OutPackagesWithSize.ValueSort(TGreater<uint64>());
+
+	return true;
+}
+
+bool GetFilteredPackagesCategorizedWithSize(TMap<FName, uint64>& OutPackagesWithSize, FName GroupName /*= NAME_None*/, FString AssetSubstring /*= FString()*/, FName ClassName /*= NAME_None*/, FName CategoryName /*= NAME_None*/, FOutputDevice* ErrorOutput /*= GLog*/)
+{
+	TArray<FLLMTagSetAllocationFilter> Filters;
+
+	FName LongName = NAME_None;
+	if (!AssetSubstring.IsEmpty())
+	{
+		Internal::GetLongName(AssetSubstring, LongName, ErrorOutput);
+	}
+
+	if (LongName != NAME_None)
+	{
+		Filters.Add({ LongName, ELLMTagSet::Assets });
+	}
+
+	if (GroupName != NAME_None)
+	{
+		Filters.Add({ GroupName, ELLMTagSet::None });
+	}
+
+	if (ClassName != NAME_None)
+	{
+		Filters.Add({ ClassName, ELLMTagSet::AssetClasses });
+	}
+
+	ELLMTagSet TagSetCategory = ELLMTagSet::None;
+	if (CategoryName == FName("Assets"))
+	{
+		TagSetCategory = ELLMTagSet::Assets;
+	}
+	else  if (CategoryName == FName("AssetClasses"))
+	{
+		TagSetCategory = ELLMTagSet::AssetClasses;
+	}
+
+	Internal::GMemoryUsageInfoProviderLLM.GetFilteredTagsWithSize(OutPackagesWithSize, ELLMTracker::Default, TagSetCategory, Filters);
+
+	if (AssetSubstring.Len() > 0)
+	{
+		Internal::RemoveFilteredPackages(OutPackagesWithSize, AssetSubstring);
+	}
+
+	OutPackagesWithSize.ValueSort(TGreater<uint64>());
+
+	return true;
+}
+
+bool GetFilteredClassesWithSize(TMap<FName, uint64>& OutClassesWithSize, FName GroupName /* = NAME_None */, FString AssetName /* = FString() */, FOutputDevice* ErrorOutput /* = GLog */)
+{
+	TArray<FLLMTagSetAllocationFilter> Filters;
+
+	FName LongName = NAME_None;
+	if (!AssetName.IsEmpty() && !Internal::GetLongName(AssetName, LongName, ErrorOutput))
+	{
+		return false;
+	}
+
+	if (LongName != NAME_None)
+	{
+		Filters.Add({LongName, ELLMTagSet::Assets});
+	}
+
+	if (GroupName != NAME_None)
+	{
+		Filters.Add({GroupName, ELLMTagSet::None});
+	}
+
+	Internal::GMemoryUsageInfoProviderLLM.GetFilteredTagsWithSize(OutClassesWithSize, ELLMTracker::Default, ELLMTagSet::AssetClasses, Filters);
+
+	OutClassesWithSize.ValueSort(TGreater<uint64>());
+
+	return true;
+}
+
+bool GetFilteredGroupsWithSize(TMap<FName, uint64>& OutGroupsWithSize, FString AssetName /* = FString() */, FName ClassName /* = NAME_None */, FOutputDevice* ErrorOutput /* = GLog */)
+{
+	TArray<FLLMTagSetAllocationFilter> Filters;
+
+	FName LongName = NAME_None;
+	if (!AssetName.IsEmpty() && !Internal::GetLongName(AssetName, LongName, ErrorOutput))
+	{
+		return false;
+	}
+
+	if (LongName != NAME_None)
+	{
+		Filters.Add({LongName, ELLMTagSet::Assets});
+	}
+
+	if (ClassName != NAME_None)
+	{
+		Filters.Add({ClassName, ELLMTagSet::AssetClasses});
+	}
+
+	Internal::GMemoryUsageInfoProviderLLM.GetFilteredTagsWithSize(OutGroupsWithSize, ELLMTracker::Default, ELLMTagSet::None, Filters);
+
+	OutGroupsWithSize.ValueSort(TGreater<uint64>());
+
+	return true;
+}
+
+} // namespace MemoryUsageQueries
+#else // ENABLE_LOW_LEVEL_MEM_TRACKER
+namespace MemoryUsageQueries
+{
+// For now only LLM backend is supported
+IMemoryUsageInfoProvider* GetCurrentMemoryUsageInfoProvider()
+{
+	return nullptr;
+}
+} // namespace MemoryUsageQueries
+#endif // ENABLE_LOW_LEVEL_MEM_TRACKER
diff --git a/Engine/Plugins/MemoryUsageQueries/Source/MemoryUsageQueries/Private/MemoryUsageInfoProviderLLM.h b/Engine/Plugins/MemoryUsageQueries/Source/MemoryUsageQueries/Private/MemoryUsageInfoProviderLLM.h
new file mode 100644
--- /dev/null
+++ b/Engine/Plugins/MemoryUsageQueries/Source/MemoryUsageQueries/Private/MemoryUsageInfoProviderLLM.h
@@ -0,0 +1,23 @@
+// Copyright Epic Games, Inc. All Rights Reserved.
+
+#pragma once
+
+#include "CoreGlobals.h"
+#include "MemoryUsageInfoProvider.h"
+#include "HAL/LowLevelMemTracker.h"
+#include "Misc/OutputDeviceRedirector.h"
+
+#if ENABLE_LOW_LEVEL_MEM_TRACKER
+
+class MEMORYUSAGEQUERIES_API FMemoryUsageInfoProviderLLM : public IMemoryUsageInfoProvider
+{
+public:
+	virtual bool IsProviderAvailable() const override;
+	virtual uint64 GetAssetMemoryUsage(FName Asset) const override;
+	virtual uint64 GetAssetsMemoryUsage(const TSet<FName>& Assets) const override;
+	virtual void GetAllAssetsWithSize(TMap<FName, uint64>& OutAssets) const override;
+	virtual uint64 GetAssetsMemoryUsageWithSize(const TSet<FName>& Assets, TMap<FName, uint64>& OutSizes) const override;
+	virtual void GetFilteredTagsWithSize(TMap<FName, uint64>& OutTags, ELLMTracker Tracker, ELLMTagSet TagSet, TArray<FLLMTagSetAllocationFilter>& Filters) const;
+};
+
+#endif // ENABLE_LOW_LEVEL_MEM_TRACKER
diff --git a/Engine/Plugins/MemoryUsageQueries/Source/MemoryUsageQueries/Private/MemoryUsageQueries.cpp b/Engine/Plugins/MemoryUsageQueries/Source/MemoryUsageQueries/Private/MemoryUsageQueries.cpp
--- a/Engine/Plugins/MemoryUsageQueries/Source/MemoryUsageQueries/Private/MemoryUsageQueries.cpp
+++ b/Engine/Plugins/MemoryUsageQueries/Source/MemoryUsageQueries/Private/MemoryUsageQueries.cpp
@@ -23,1149 +23,312 @@
 namespace MemoryUsageQueries
 {
 
-FMemoryUsageInfoProviderLLM MemoryUsageInfoProviderLLM;
-IMemoryUsageInfoProvider* CurrentMemoryUsageInfoProvider = &MemoryUsageInfoProviderLLM;
-
-} // namespace MemoryUsageQueries
-
-struct FAssetMemoryBreakdown
-{
-	FAssetMemoryBreakdown()
-	{
-	}
-
-	uint64 ExclusiveSize = 0;
-	uint64 UniqueSize = 0;
-	uint64 SharedSize = 0;
-	uint64 TotalSize = 0;
-};
-
-struct FAssetMemoryDetails
+void GetMemoryUsage(const IMemoryUsageInfoProvider* MemoryUsageInfoProvider, const FName& PackageName, uint64& OutExclusiveSize, uint64& OutInclusiveSize)
 {
-	FAssetMemoryDetails()
-	{
-	}
-
-	// Assets Package Name
-	FName PackageName;
-
-	FAssetMemoryBreakdown MemoryBreakdown;
-
-	// List of dependencies for this asset
 	TSet<FName> Dependencies;
-
-	TMap<FName, FAssetMemoryBreakdown> DependenciesToMemoryMap;
-
-	int32 UniqueRefCount = 0;
-	int32 SharedRefCount = 0;
-};
-
-static int32 DefaultResultLimit = 15;
-
-// Helper struct that scopes an output device, and will automatically clean up any log files when leaving scope
-struct FScopedOutputDevice
-{
-	FOutputDevice* CurrentOutputDevice = nullptr;
-
-	FOutputDeviceArchiveWrapper* FileArWrapper = nullptr;
-	FArchive* FileAr = nullptr;
-
-	explicit FScopedOutputDevice(FOutputDevice* InDefaultOutputDevice) :
-		CurrentOutputDevice(InDefaultOutputDevice)
-		
-	{
-	}
-
-	~FScopedOutputDevice()
-	{
-		if (FileArWrapper != nullptr)
-		{
-			FileArWrapper->TearDown();
-		}
-
-		delete FileArWrapper;
-		delete FileAr;
-	}
-
-	FOutputDevice& GetOutputDevice() const
-	{
-		return *CurrentOutputDevice;
-	}
-
-	void OpenLogFile(const FString& LogFileName, bool bCSV)
-	{
-		if (FileArWrapper != nullptr || LogFileName.IsEmpty())
-		{
-			// Log file is already open or LogFileName is empty
-			return;
-		}
-
-#if ALLOW_DEBUG_FILES
-		// Create folder for MemQuery files.
-		const FString OutputDir = FPaths::ProfilingDir() / TEXT("MemQuery");
-		IFileManager::Get().MakeDirectory(*OutputDir, true);
-
-		FString FileTimeString = FString::Printf(TEXT("_%s"), *FDateTime::Now().ToString(TEXT("%H%M%S")));
-		const FString FileExtension = (bCSV ? TEXT(".csv") : TEXT(".memquery"));
-		const FString LogFilename = OutputDir / (LogFileName + FileTimeString + FileExtension);
-
-		FileAr = IFileManager::Get().CreateDebugFileWriter(*LogFilename);
-
-		if (FileAr != nullptr)
-		{
-			FileArWrapper = new FOutputDeviceArchiveWrapper(FileAr);
-				
-			if (FileArWrapper != nullptr)
-			{
-				CurrentOutputDevice = FileArWrapper;
-			}
-			else
-			{
-				delete FileAr;
-				FileAr = nullptr;
-			}
-		}
-#endif
-	}
-
-	// Delete default constructor to force proper usage
-	FScopedOutputDevice() = delete;
-};
-
-/** Structure that will parse a string and fill out some commonly used parameters per command */
-struct FCommonParameters
-{
-	int32 Limit = -1;
-
-	bool bTruncate = true;
-	bool bCSV = false;
-	
-	FName Group = NAME_None;
-	FName Class = NAME_None;
-	FName Category = NAME_None;
-
-	FString Name;
-	FString Names;
-	FString AssetName;
-	FString LogFileName;
-
-	explicit FCommonParameters(const TCHAR* Args)
-	{
-		bTruncate = !FParse::Param(Args, TEXT("notrunc"));
-		bCSV = FParse::Param(Args, TEXT("csv"));
-
-		FParse::Value(Args, TEXT("Name="), Name);
-		FParse::Value(Args, TEXT("Names="), Names);
-		FParse::Value(Args, TEXT("Limit="), Limit);
-		FParse::Value(Args, TEXT("Log="), LogFileName);
-		FParse::Value(Args, TEXT("Asset="), AssetName);
-
-		FString GroupName;
-		if (FParse::Value(Args, TEXT("Group="), GroupName))
-		{
-			Group = FName(*GroupName);
-		}
-
-		FString ClassName;
-		if (FParse::Value(Args, TEXT("Class="), ClassName))
-		{
-			Class = FName(*ClassName);
-		}
-
-		FString CategoryName;
-		if (FParse::Value(Args, TEXT("Category="), CategoryName))
-		{
-			Category = FName(*CategoryName);
-		}
-	}
-
-	FCommonParameters() = delete;
-};
-
-FAutoConsoleCommandWithWorldArgsAndOutputDevice GMemQueryUsage(
-	TEXT("MemQuery.Usage"),
-	TEXT("Name=<AssetName> Prints memory usage of the specified asset."),
-	FConsoleCommandWithWorldArgsAndOutputDeviceDelegate::CreateStatic(
-		[](const TArray<FString>& Args, UWorld* World, FOutputDevice& Ar)
-{
-	const FString Cmd = FString::Join(Args, TEXT(" "));
-
-	FScopedOutputDevice ScopedOutputDevice(&Ar);
-	FCommonParameters CommonArgs(*Cmd);
-
-	ScopedOutputDevice.OpenLogFile(CommonArgs.LogFileName, CommonArgs.bCSV);
-
-	if (!CommonArgs.Name.IsEmpty())
-	{
-		uint64 ExclusiveSize = 0U;
-		uint64 InclusiveSize = 0U;
-
-		if (MemoryUsageQueries::GetMemoryUsage(MemoryUsageQueries::CurrentMemoryUsageInfoProvider, CommonArgs.Name, ExclusiveSize, InclusiveSize, &Ar))
-		{
-			ScopedOutputDevice.GetOutputDevice().Logf(TEXT("MemoryUsage: ExclusiveSize: %.2f MiB (%.2f KiB); InclusiveSize: %.2f MiB (%.2f KiB)"),
-				ExclusiveSize / (1024.f * 1024.f),
-				ExclusiveSize / 1024.f,
-				InclusiveSize / (1024.f * 1024.f),
-				InclusiveSize / 1024.f);
-		}
-	}
-}));
-
-FAutoConsoleCommandWithWorldArgsAndOutputDevice GMemQueryCombinedUsage(
-	TEXT("MemQuery.CombinedUsage"),
-	TEXT("Names=\"<AssetName1> <AssetName2> ...\" Prints combined memory usage of the specified assets (including all dependencies)."),
-	FConsoleCommandWithWorldArgsAndOutputDeviceDelegate::CreateStatic(
-		[](const TArray<FString>& Args, UWorld* World, FOutputDevice& Ar)
-{
-	const FString Cmd = FString::Join(Args, TEXT(" "));
-
-	FScopedOutputDevice ScopedOutputDevice(&Ar);
-	FCommonParameters CommonArgs(*Cmd);
-
-	ScopedOutputDevice.OpenLogFile(CommonArgs.LogFileName, CommonArgs.bCSV);
-
-	if (!CommonArgs.Names.IsEmpty())
-	{
-		TArray<FString> Packages;
-		CommonArgs.Names.ParseIntoArrayWS(Packages);
-		uint64 TotalSize;
-
-		if (MemoryUsageQueries::GetMemoryUsageCombined(MemoryUsageQueries::CurrentMemoryUsageInfoProvider, Packages, TotalSize, &Ar))
-		{
-			ScopedOutputDevice.GetOutputDevice().Logf(TEXT("MemoryUsageCombined: TotalSize: %.2f MiB (%.2f KiB)"), TotalSize / (1024.f * 1024.f), TotalSize / 1024.f);
-		}
-	}
-}));
-
-FAutoConsoleCommandWithWorldArgsAndOutputDevice GMemQuerySharedUsage(
-	TEXT("MemQuery.SharedUsage"),
-	TEXT("Names=\"<AssetName1> <AssetName2> ...\" Prints shared memory usage of the specified assets (including only dependencies shared by the specified assets)."),
-	FConsoleCommandWithWorldArgsAndOutputDeviceDelegate::CreateStatic(
-		[](const TArray<FString>& Args, UWorld* World, FOutputDevice& Ar)
-{
-	const FString Cmd = FString::Join(Args, TEXT(" "));
-
-	FScopedOutputDevice ScopedOutputDevice(&Ar);
-	FCommonParameters CommonArgs(*Cmd);
-
-	ScopedOutputDevice.OpenLogFile(CommonArgs.LogFileName, CommonArgs.bCSV);
-
-	if (!CommonArgs.Names.IsEmpty())
-	{
-		TArray<FString> Packages;
-		CommonArgs.Names.ParseIntoArrayWS(Packages);
-		uint64 SharedSize;
-
-		if (MemoryUsageQueries::GetMemoryUsageShared(MemoryUsageQueries::CurrentMemoryUsageInfoProvider, Packages, SharedSize, &Ar))
-		{
-			ScopedOutputDevice.GetOutputDevice().Logf(TEXT("MemoryUsageShared: SharedSize: %.2f MiB (%.2f KiB)"), SharedSize / (1024.f * 1024.f), SharedSize / 1024.f);
-		}
-	}
-}));
-
-FAutoConsoleCommandWithWorldArgsAndOutputDevice GMemQueryUniqueUsage(
-	TEXT("MemQuery.UniqueUsage"),
-	TEXT("Names=\"<AssetName1> <AssetName2> ...\" Prints unique memory usage of the specified assets (including only dependencies unique to the specified assets)."),
-	FConsoleCommandWithWorldArgsAndOutputDeviceDelegate::CreateStatic(
-		[](const TArray<FString>& Args, UWorld* World, FOutputDevice& Ar)
-{
-	const FString Cmd = FString::Join(Args, TEXT(" "));
-
-	FScopedOutputDevice ScopedOutputDevice(&Ar);
-	FCommonParameters CommonArgs(*Cmd);
-
-	ScopedOutputDevice.OpenLogFile(CommonArgs.LogFileName, CommonArgs.bCSV);
-
-	if (!CommonArgs.Names.IsEmpty())
-	{
-		TArray<FString> Packages;
-		CommonArgs.Names.ParseIntoArrayWS(Packages);
-		uint64 UniqueSize = 0U;
-
-		if (MemoryUsageQueries::GetMemoryUsageUnique(MemoryUsageQueries::CurrentMemoryUsageInfoProvider, Packages, UniqueSize, &Ar))
-		{
-			ScopedOutputDevice.GetOutputDevice().Logf(TEXT("MemoryUsageUnique: UniqueSize: %.2f MiB (%.2f KiB)"), UniqueSize / (1024.f * 1024.f), UniqueSize / 1024.f);
-		}
-	}
-}));
-
-FAutoConsoleCommandWithWorldArgsAndOutputDevice GMemQueryCommonUsage(
-	TEXT("MemQuery.CommonUsage"),
-	TEXT("Names=\"<AssetName1> <AssetName2> ...\" Prints common memory usage of the specified assets (including only dependencies that are not unique to the specified assets)."),
-	FConsoleCommandWithWorldArgsAndOutputDeviceDelegate::CreateStatic(
-		[](const TArray<FString>& Args, UWorld* World, FOutputDevice& Ar)
-{
-	const FString Cmd = FString::Join(Args, TEXT(" "));
-
-	FScopedOutputDevice ScopedOutputDevice(&Ar);
-	FCommonParameters CommonArgs(*Cmd);
-
-	ScopedOutputDevice.OpenLogFile(CommonArgs.LogFileName, CommonArgs.bCSV);
-
-	if (!CommonArgs.Names.IsEmpty())
-	{
-		TArray<FString> Packages;
-		CommonArgs.Names.ParseIntoArrayWS(Packages);
-		uint64 CommonSize = 0U;
-
-		if (MemoryUsageQueries::GetMemoryUsageCommon(MemoryUsageQueries::CurrentMemoryUsageInfoProvider, Packages, CommonSize, &Ar))
-		{
-			ScopedOutputDevice.GetOutputDevice().Logf(TEXT("MemoryUsageCommon: CommonSize: %.2f MiB (%.2f KiB)"), CommonSize / (1024.f * 1024.f), CommonSize / 1024.f);
-		}
-
-	}
-}));
-
-FAutoConsoleCommandWithWorldArgsAndOutputDevice GMemQueryDependencies(
-	TEXT("MemQuery.Dependencies"),
-	TEXT("Name=<AssetName> Limit=<n> Lists dependencies of the specified asset, sorted by size."),
-	FConsoleCommandWithWorldArgsAndOutputDeviceDelegate::CreateStatic(
-		[](const TArray<FString>& Args, UWorld* World, FOutputDevice& Ar)
-{
-	const FString Cmd = FString::Join(Args, TEXT(" "));
-
-	FScopedOutputDevice ScopedOutputDevice(&Ar);
-	FCommonParameters CommonArgs(*Cmd);
-
-	ScopedOutputDevice.OpenLogFile(CommonArgs.LogFileName, CommonArgs.bCSV);
-
-	if (!CommonArgs.Name.IsEmpty())
-	{
-		TMap<FName, uint64> DependenciesWithSize;
-
-		if (MemoryUsageQueries::GetDependenciesWithSize(MemoryUsageQueries::CurrentMemoryUsageInfoProvider, CommonArgs.Name, DependenciesWithSize, &Ar))
-		{
-			MemoryUsageQueries::Internal::PrintTagsWithSize(ScopedOutputDevice.GetOutputDevice(), DependenciesWithSize, TEXT("Dependencies"), CommonArgs.bTruncate, CommonArgs.Limit, CommonArgs.bCSV);
-		}
-	}
-}));
-
-FAutoConsoleCommandWithWorldArgsAndOutputDevice GMemQueryCombinedDependencies(
-	TEXT("MemQuery.CombinedDependencies"),
-	TEXT("Names=\"<AssetName1> <AssetName2> ...\" Limit=<n> Lists n largest dependencies of the specified assets, sorted by size."),
-	FConsoleCommandWithWorldArgsAndOutputDeviceDelegate::CreateStatic(
-		[](const TArray<FString>& Args, UWorld* World, FOutputDevice& Ar)
-{
-	const FString Cmd = FString::Join(Args, TEXT(" "));
-
-	FScopedOutputDevice ScopedOutputDevice(&Ar);
-	FCommonParameters CommonArgs(*Cmd);
-
-	ScopedOutputDevice.OpenLogFile(CommonArgs.LogFileName, CommonArgs.bCSV);
-
-	if (!CommonArgs.Names.IsEmpty())
-	{
-		TArray<FString> Packages;
-		CommonArgs.Names.ParseIntoArrayWS(Packages);
-		TMap<FName, uint64> CombinedDependenciesWithSize;
-
-		if (MemoryUsageQueries::GetDependenciesWithSizeCombined(MemoryUsageQueries::CurrentMemoryUsageInfoProvider, Packages, CombinedDependenciesWithSize, &Ar))
-		{
-			MemoryUsageQueries::Internal::PrintTagsWithSize(ScopedOutputDevice.GetOutputDevice(), CombinedDependenciesWithSize, TEXT("Combined Dependencies"), CommonArgs.bTruncate, CommonArgs.Limit, CommonArgs.bCSV);
-		}
-	}
-}));
-
-FAutoConsoleCommandWithWorldArgsAndOutputDevice GMemQuerySharedDependencies(
-	TEXT("MemQuery.SharedDependencies"),
-	TEXT("Names=\"<AssetName1> <AssetName2> ...\" Limit=<n> Lists n largest dependencies that are shared by the specified assets, sorted by size."),
-	FConsoleCommandWithWorldArgsAndOutputDeviceDelegate::CreateStatic(
-		[](const TArray<FString>& Args, UWorld* World, FOutputDevice& Ar)
-{
-	const FString Cmd = FString::Join(Args, TEXT(" "));
-
-	FScopedOutputDevice ScopedOutputDevice(&Ar);
-	FCommonParameters CommonArgs(*Cmd);
-
-	ScopedOutputDevice.OpenLogFile(CommonArgs.LogFileName, CommonArgs.bCSV);
-
-	if (!CommonArgs.Names.IsEmpty())
-	{
-		TArray<FString> Packages;
-		CommonArgs.Names.ParseIntoArrayWS(Packages);
-		TMap<FName, uint64> SharedDependenciesWithSize;
-
-		if (MemoryUsageQueries::GetDependenciesWithSizeShared(MemoryUsageQueries::CurrentMemoryUsageInfoProvider, Packages, SharedDependenciesWithSize, &Ar))
-		{
-			MemoryUsageQueries::Internal::PrintTagsWithSize(ScopedOutputDevice.GetOutputDevice(), SharedDependenciesWithSize, TEXT("Shared Dependencies"), CommonArgs.bTruncate, CommonArgs.Limit, CommonArgs.bCSV);
-		}
-	}
-}));
-
-FAutoConsoleCommandWithWorldArgsAndOutputDevice GMemQueryUniqueDependencies(
-	TEXT("MemQuery.UniqueDependencies"),
-	TEXT("Names=\"<AssetName1> <AssetName2> ...\" Limit=<n> Lists n largest dependencies that are unique to the specified assets, sorted by size."),
-	FConsoleCommandWithWorldArgsAndOutputDeviceDelegate::CreateStatic(
-		[](const TArray<FString>& Args, UWorld* World, FOutputDevice& Ar)
-{
-	const FString Cmd = FString::Join(Args, TEXT(" "));
-
-	FScopedOutputDevice ScopedOutputDevice(&Ar);
-	FCommonParameters CommonArgs(*Cmd);
-
-	ScopedOutputDevice.OpenLogFile(CommonArgs.LogFileName, CommonArgs.bCSV);
-
-	if (!CommonArgs.Names.IsEmpty())
-	{
-		TArray<FString> Packages;
-		CommonArgs.Names.ParseIntoArrayWS(Packages);
-		TMap<FName, uint64> UniqueDependenciesWithSize;
-
-		if (MemoryUsageQueries::GetDependenciesWithSizeUnique(MemoryUsageQueries::CurrentMemoryUsageInfoProvider, Packages, UniqueDependenciesWithSize, &Ar))
-		{
-			MemoryUsageQueries::Internal::PrintTagsWithSize(ScopedOutputDevice.GetOutputDevice(), UniqueDependenciesWithSize, TEXT("Unique Dependencies"), CommonArgs.bTruncate, CommonArgs.Limit, CommonArgs.bCSV);
-		}
-	}
-}));
-
-FAutoConsoleCommandWithWorldArgsAndOutputDevice GMemQueryCommonDependencies(
-	TEXT("MemQuery.CommonDependencies"),
-	TEXT("Names=\"<AssetName1> <AssetName2> ...\" Limit=<n> Lists n largest dependencies that are NOT unique to the specified assets, sorted by size."),
-	FConsoleCommandWithWorldArgsAndOutputDeviceDelegate::CreateStatic(
-		[](const TArray<FString>& Args, UWorld* World, FOutputDevice& Ar)
-{
-	const FString Cmd = FString::Join(Args, TEXT(" "));
-
-	FScopedOutputDevice ScopedOutputDevice(&Ar);
-	FCommonParameters CommonArgs(*Cmd);
-
-	ScopedOutputDevice.OpenLogFile(CommonArgs.LogFileName, CommonArgs.bCSV);
-
-	if (!CommonArgs.Names.IsEmpty())
-	{
-		TArray<FString> Packages;
-		CommonArgs.Names.ParseIntoArrayWS(Packages);
-		TMap<FName, uint64> CommonDependenciesWithSize;
-
-		if (MemoryUsageQueries::GetDependenciesWithSizeCommon(MemoryUsageQueries::CurrentMemoryUsageInfoProvider, Packages, CommonDependenciesWithSize, &Ar))
-		{
-			MemoryUsageQueries::Internal::PrintTagsWithSize(ScopedOutputDevice.GetOutputDevice(), CommonDependenciesWithSize, TEXT("Common Dependencies"), CommonArgs.bTruncate, CommonArgs.Limit, CommonArgs.bCSV);
-		}
-	}
-}));
-
-FAutoConsoleCommandWithWorldArgsAndOutputDevice GMemQuerySavings(
-	TEXT("MemQuery.Savings"),
-	TEXT("Limit=<n> Lists potential savings among %s. How much memory can be saved it we delete certain object."),
-	FConsoleCommandWithWorldArgsAndOutputDeviceDelegate::CreateStatic(
-		[](const TArray<FString>& Args, UWorld* World, FOutputDevice& Ar)
-{
-	const FString Cmd = FString::Join(Args, TEXT(" "));
-
-	FScopedOutputDevice ScopedOutputDevice(&Ar);
-	FCommonParameters CommonArgs(*Cmd);
-
-	ScopedOutputDevice.OpenLogFile(CommonArgs.LogFileName, CommonArgs.bCSV);
-
-	const UMemoryUsageQueriesConfig* Config = GetDefault<UMemoryUsageQueriesConfig>();
-
-	// Necessary for FParse::Command
-	const TCHAR* Command = *Cmd;
-
-	for (auto It = Config->SavingsPresets.CreateConstIterator(); It; ++It)
-	{
-		if (!FParse::Command(&Command, *It.Key()))
-		{
-			continue;
-		}
-
-		TMap<FName, uint64> PresetSavings;
-		TSet<FName> Packages;
-
-		UClass* SavingsClass = FindObject<UClass>(nullptr, *It.Value());
-		if (SavingsClass != nullptr)
-		{
-			TArray<UClass*> DerivedClasses;
-			TArray<UClass*> DerivedResults;
-
-			GetDerivedClasses(SavingsClass, DerivedClasses, true);
-
-			for (UClass* DerivedClass : DerivedClasses)
-			{
-				UBlueprintGeneratedClass* BPClass = Cast<UBlueprintGeneratedClass>(DerivedClass);
-				if (BPClass != nullptr)
-				{
-					DerivedResults.Reset();
-					GetDerivedClasses(BPClass, DerivedResults, false);
-					if (DerivedResults.Num() == 0)
-					{
-						Packages.Add(DerivedClass->GetPackage()->GetFName());
-					}
-				}
-			}
-		}
-
-		for (const auto& Package : Packages)
-		{
-			uint64 Size = 0;
-			MemoryUsageQueries::GetMemoryUsageUnique(MemoryUsageQueries::CurrentMemoryUsageInfoProvider, TArray<FString>({ Package.ToString() }), Size, &Ar);
-			PresetSavings.Add(Package, Size);
-		}
-
-		PresetSavings.ValueSort(TGreater<uint64>());
-		MemoryUsageQueries::Internal::PrintTagsWithSize(ScopedOutputDevice.GetOutputDevice(), PresetSavings, *FString::Printf(TEXT("possible savings")), CommonArgs.bTruncate, CommonArgs.bCSV);
-	}
-}));
-
-#if ENABLE_LOW_LEVEL_MEM_TRACKER
-
-FAutoConsoleCommandWithWorldArgsAndOutputDevice GMemQueryListAssets(
-	TEXT("MemQuery.ListAssets"),
-	TEXT("Asset=<AssetNameSubstring> Group=<GroupName> Class=<ClassName> Limit=<n> Lists n largest assets."),
-	FConsoleCommandWithWorldArgsAndOutputDeviceDelegate::CreateStatic(
-		[](const TArray<FString>& Args, UWorld* World, FOutputDevice& Ar)
-{
-	const FString Cmd = FString::Join(Args, TEXT(" "));
-
-	FScopedOutputDevice ScopedOutputDevice(&Ar);
-	FCommonParameters CommonArgs(*Cmd);
-
-	ScopedOutputDevice.OpenLogFile(CommonArgs.LogFileName, CommonArgs.bCSV);
-
-	bool bSuccess;
-	TMap<FName, uint64> AssetsWithSize;
-
-	if (CommonArgs.Group != NAME_None || CommonArgs.Class != NAME_None)
-	{
-		bSuccess = MemoryUsageQueries::GetFilteredPackagesWithSize(AssetsWithSize, CommonArgs.Group, CommonArgs.AssetName, CommonArgs.Class, &ScopedOutputDevice.GetOutputDevice());
-	}
-	else
-	{
-		// TODO - Implement using faster path if there are no group / class filters
-		bSuccess = MemoryUsageQueries::GetFilteredPackagesWithSize(AssetsWithSize, CommonArgs.Group, CommonArgs.AssetName, CommonArgs.Class, &ScopedOutputDevice.GetOutputDevice());
-	}
-
-	if (bSuccess)
-	{
-		MemoryUsageQueries::Internal::PrintTagsWithSize(ScopedOutputDevice.GetOutputDevice(), AssetsWithSize, TEXT("largest assets"), CommonArgs.bTruncate, CommonArgs.Limit, CommonArgs.bCSV);
-	}
-}));
-
-FAutoConsoleCommandWithWorldArgsAndOutputDevice GMemQueryListAssetsCategorized(
-	TEXT("MemQuery.ListAssetsCategorized"),
-	TEXT("Asset=<AssetNameSubstring> Group=<GroupName> Class=<ClassName> Category=<CategoryName(None,Assets,AssetClasses)> Limit=<n> Lists n largest assets categorized by Category."),
-	FConsoleCommandWithWorldArgsAndOutputDeviceDelegate::CreateStatic(
-		[](const TArray<FString>& Args, UWorld* World, FOutputDevice& Ar)
-{
-	const FString Cmd = FString::Join(Args, TEXT(" "));
-
-	FScopedOutputDevice ScopedOutputDevice(&Ar);
-	FCommonParameters CommonArgs(*Cmd);
-
-	ScopedOutputDevice.OpenLogFile(CommonArgs.LogFileName, CommonArgs.bCSV);
-
-	bool bSuccess;
-	TMap<FName, uint64> AssetsWithSize;						
-	bSuccess = MemoryUsageQueries::GetFilteredPackagesCategorizedWithSize(AssetsWithSize, CommonArgs.Group, CommonArgs.AssetName, CommonArgs.Class, CommonArgs.Category, &ScopedOutputDevice.GetOutputDevice());
-
-	if (bSuccess)
-	{
-		MemoryUsageQueries::Internal::PrintTagsWithSize(ScopedOutputDevice.GetOutputDevice(), AssetsWithSize, TEXT("largest assets"), CommonArgs.bTruncate, CommonArgs.Limit, CommonArgs.bCSV);
-	}
-}));
-
-FAutoConsoleCommandWithWorldArgsAndOutputDevice GMemQueryListClasses(
-	TEXT("MemQuery.ListClasses"),
-	TEXT("Group=<GroupName> Asset=<AssetName> Limit=<n> Lists n largest classes."),
-	FConsoleCommandWithWorldArgsAndOutputDeviceDelegate::CreateStatic(
-		[](const TArray<FString>& Args, UWorld* World, FOutputDevice& Ar)
-{
-	const FString Cmd = FString::Join(Args, TEXT(" "));
-
-	FScopedOutputDevice ScopedOutputDevice(&Ar);
-	FCommonParameters CommonArgs(*Cmd);
-
-	ScopedOutputDevice.OpenLogFile(CommonArgs.LogFileName, CommonArgs.bCSV);
-
-	TMap<FName, uint64> ClassesWithSize;
-	if (MemoryUsageQueries::GetFilteredClassesWithSize(ClassesWithSize, CommonArgs.Group, CommonArgs.AssetName, &Ar))
-	{
-		MemoryUsageQueries::Internal::PrintTagsWithSize(ScopedOutputDevice.GetOutputDevice(), ClassesWithSize, *FString::Printf(TEXT("Largest Classes")), CommonArgs.bTruncate, CommonArgs.Limit, CommonArgs.bCSV);
-	}
-}));
-
-FAutoConsoleCommandWithWorldArgsAndOutputDevice GMemQueryListGroups(
-	TEXT("MemQuery.ListGroups"),
-	TEXT("Asset=<AssetName> Class=<ClassName> Limit=<n> Lists n largest groups."),
-	FConsoleCommandWithWorldArgsAndOutputDeviceDelegate::CreateStatic(
-		[](const TArray<FString>& Args, UWorld* World, FOutputDevice& Ar)
-{
-	const FString Cmd = FString::Join(Args, TEXT(" "));
-
-	FScopedOutputDevice ScopedOutputDevice(&Ar);
-	FCommonParameters CommonArgs(*Cmd);
-
-	ScopedOutputDevice.OpenLogFile(CommonArgs.LogFileName, CommonArgs.bCSV);
-
-	TMap<FName, uint64> GroupsWithSize;
-	if (MemoryUsageQueries::GetFilteredGroupsWithSize(GroupsWithSize, CommonArgs.AssetName, CommonArgs.Class, &Ar))
-	{
-		MemoryUsageQueries::Internal::PrintTagsWithSize(ScopedOutputDevice.GetOutputDevice(), GroupsWithSize, *FString::Printf(TEXT("Largest Groups")), CommonArgs.bTruncate, CommonArgs.Limit, CommonArgs.bCSV);
-	}
-}));
-
-// Will return true if the Package Name matches any of the conditions in the array of Paths
-static bool PackageNameMatches(const FString& PackageName, const TArray<FString>& Conditions)
-{
-	for (const FString& Condition : Conditions)
-	{
-		if ((FWildcardString::ContainsWildcards(*Condition) && FWildcardString::IsMatch(*Condition, *PackageName)) || (PackageName.Contains(Condition)))
-		{
-			return true;
-		}
-	}
-
-	return false;
+	Internal::GetTransitiveDependencies(PackageName, Dependencies);
+	Dependencies.Add(PackageName);
+	OutExclusiveSize = MemoryUsageInfoProvider->GetAssetMemoryUsage(PackageName);
+	OutInclusiveSize = MemoryUsageInfoProvider->GetAssetsMemoryUsage(Dependencies);
 }
-
-FAutoConsoleCommandWithWorldArgsAndOutputDevice GMemQueryCollections(
-	TEXT("MemQuery.Collection"),
-	TEXT("Lists memory used by a collection. Can show dependency breakdown. Pass -showdeps to list dependencies."),
-	FConsoleCommandWithWorldArgsAndOutputDeviceDelegate::CreateStatic(
-		[](const TArray<FString>& Args, UWorld* World, FOutputDevice& Ar)
+bool GetMemoryUsage(const IMemoryUsageInfoProvider* MemoryUsageInfoProvider, const FString& PackageName, uint64& OutExclusiveSize, uint64& OutInclusiveSize, FOutputDevice* ErrorOutput /* = GLog */)
 {
-	const FString Cmd = FString::Join(Args, TEXT(" "));
-
-	FScopedOutputDevice ScopedOutputDevice(&Ar);
-	FCommonParameters CommonArgs(*Cmd);
-
-	ScopedOutputDevice.OpenLogFile(CommonArgs.LogFileName, CommonArgs.bCSV);
-
-	const bool bShowDependencies = FParse::Param(*Cmd, TEXT("ShowDeps"));
-
-	// Necessary for FParse::Command
-	const TCHAR* Command = *Cmd;
-
-	FOutputDevice* CurrentOutputDevice = &ScopedOutputDevice.GetOutputDevice();
-
-	const UMemoryUsageQueriesConfig* Config = GetDefault<UMemoryUsageQueriesConfig>();
-	for (auto It = Config->Collections.CreateConstIterator(); It; ++It)
+	if (!MemoryUsageInfoProvider->IsProviderAvailable())
 	{
-		const FCollectionInfo& CollectionInfo = *It;
-
-		if (!FParse::Command(&Command, *CollectionInfo.Name))
+		if (ErrorOutput)
 		{
-			continue;
-		}
-
-		// Retrieve a list of all assets that have allocations we are currently tracking.
-		TMap<FName, uint64> AssetsWithSize;
-		bool bSuccess = MemoryUsageQueries::GetFilteredPackagesWithSize(AssetsWithSize, NAME_None, "", NAME_None, &Ar);
-
-		if (!bSuccess)
-		{
-			Ar.Logf(TEXT("Failed to gather assets for Collection %s"), *CollectionInfo.Name);
-			break;
+			ErrorOutput->Logf(ELogVerbosity::Error, TEXT("MemoryUsageInfoProvider Error: Provider is not available. Please run with -LLM"));
 		}
-
-		// See if any of the asset paths match those of our matching paths and are valid
-		TArray<FString> PackageNames;
-
-		TMap<FName, FAssetMemoryDetails> AssetMemoryMap;
-		for (const TPair<FName, uint64>& AssetSizePair : AssetsWithSize)
-		{
-			const FString& PackageName = AssetSizePair.Key.ToString();
-
-			if (!FPackageName::IsValidLongPackageName(PackageName))
-			{
-				continue;
-			}
-
-			// If path is Included and NOT Excluded, it's a valid asset to consider
-			if (PackageNameMatches(PackageName, CollectionInfo.Includes) && !PackageNameMatches(PackageName, CollectionInfo.Excludes))
-			{
-				PackageNames.Add(PackageName);
-				FAssetMemoryDetails& AssetMemory = AssetMemoryMap.Add(AssetSizePair.Key);
-				AssetMemory.MemoryBreakdown.ExclusiveSize = AssetSizePair.Value;
-
-				FName LongPackageName;
-				if (!MemoryUsageQueries::Internal::GetLongNameAndDependencies(PackageName, LongPackageName, AssetMemory.Dependencies, &Ar))
-				{
-					Ar.Logf(TEXT("Failed to get dependencies foro Asset %s"), *PackageName);
-				}
-			}
-		}
-
-		// Gather list of dependencies. Internal dependencies are confined only to the set of packages passed in.
-		// External are dependencies that have additional references outside the set of packages passed in.
-		TMap<FName, uint64> InternalDependencies;
-		TMap<FName, uint64> ExternalDependencies;
-		if (!MemoryUsageQueries::GatherDependenciesForPackages(MemoryUsageQueries::CurrentMemoryUsageInfoProvider, PackageNames, InternalDependencies, ExternalDependencies, MemoryUsageQueries::EDependencyType::EDep_All, &Ar))
-		{
-			Ar.Logf(TEXT("Failed to gather memory usage for dependencies in Collection %s"), *CollectionInfo.Name);
-			break;
-		}
-
-		uint64 TotalCollectionSize = 0;
-
-		// Determine which category where each assets dependency should reside
-		for (TPair<FName, FAssetMemoryDetails>& Asset : AssetMemoryMap)
-		{
-			FAssetMemoryDetails& AssetMemory = Asset.Value;
-			FAssetMemoryBreakdown& AssetMemoryDetails = AssetMemory.MemoryBreakdown;
-
-			for (FName& Dep : Asset.Value.Dependencies)
-			{
-				// Don't want to count asset itself, plus some dependencies might refer to other assets in the map
-				if (AssetMemoryMap.Contains(Dep))
-				{
-					continue;
-				}
-
-				FAssetMemoryBreakdown DependencyMemory;
-				uint64* UniqueMemory = InternalDependencies.Find(Dep);
-				uint64* SharedMemory = ExternalDependencies.Find(Dep);
-				bool bRecordDependency = false;
-
-				if (UniqueMemory != nullptr && *UniqueMemory != 0)
-				{
-					DependencyMemory.UniqueSize = *UniqueMemory;
-					AssetMemoryDetails.UniqueSize += DependencyMemory.UniqueSize;
-					AssetMemory.UniqueRefCount++;
-					bRecordDependency = true;
-				}
-
-				if (SharedMemory != nullptr && *SharedMemory != 0)
-				{
-					DependencyMemory.SharedSize = *SharedMemory;
-					AssetMemoryDetails.SharedSize += DependencyMemory.SharedSize;
-					AssetMemory.SharedRefCount++;
-					bRecordDependency = true;
-				}
-
-				if (bRecordDependency)
-				{
-					AssetMemory.DependenciesToMemoryMap.Add(Dep, DependencyMemory);
-				}
-			}
-
-			AssetMemoryDetails.TotalSize = AssetMemoryDetails.ExclusiveSize + AssetMemoryDetails.UniqueSize;
-			TotalCollectionSize += AssetMemoryDetails.TotalSize;
-		}
-
-		// Sort by TotalSize
-		AssetMemoryMap.ValueSort([](const FAssetMemoryDetails& A, const FAssetMemoryDetails& B) {
-			return A.MemoryBreakdown.TotalSize > B.MemoryBreakdown.TotalSize;
-			});
-
-		if (CommonArgs.bCSV)
-		{
-			CurrentOutputDevice->Logf(TEXT(",Asset,Exclusive KiB,Unique Refs KiB,Unique Ref Count,Shared Refs KiB,Shared Ref Count,Total KiB"));
-		}
-		else
-		{
-			CurrentOutputDevice->Logf(
-				TEXT(" %100s %20s %20s %15s %20s %15s %25s"),
-				TEXT("Asset"),
-				TEXT("Exclusive KiB"),
-				TEXT("Unique Refs KiB"),
-				TEXT("Unique Ref Count"),
-				TEXT("Shared Refs KiB"),
-				TEXT("Shared Ref Count"),
-				TEXT("Total KiB")
-			);
-		}
-
-		// Asset listing
-		for (TPair<FName, FAssetMemoryDetails>& Asset : AssetMemoryMap)
-		{
-			FAssetMemoryDetails& AssetMemory = Asset.Value;
-			FAssetMemoryBreakdown& AssetMemoryDetails = AssetMemory.MemoryBreakdown;
-
-			if (CommonArgs.bCSV)
-			{
-				CurrentOutputDevice->Logf(TEXT(",%s,%.2f,%.2f,%d,%.2f,%d,%.2f"), *Asset.Key.ToString(),
-					AssetMemoryDetails.ExclusiveSize / 1024.f,
-					AssetMemoryDetails.UniqueSize / 1024.f,
-					AssetMemory.UniqueRefCount,
-					AssetMemoryDetails.SharedSize / 1024.f,
-					AssetMemory.SharedRefCount,
-					AssetMemoryDetails.TotalSize / 1024.f);
-			}
-			else
-			{
-				CurrentOutputDevice->Logf(
-					TEXT(" %100s %20.2f %20.2f %15d %20.2f %15d %25.2f"),
-					*Asset.Key.ToString(),
-					AssetMemoryDetails.ExclusiveSize / 1024.f,
-					AssetMemoryDetails.SharedSize / 1024.f,
-					AssetMemory.SharedRefCount,
-					AssetMemoryDetails.UniqueSize / 1024.f,
-					AssetMemory.UniqueRefCount,
-					AssetMemoryDetails.TotalSize / 1024.f
-				);
-			}
-		}
-
-		// Asset dependencies listing
-		if (bShowDependencies)
-		{
-			if (CommonArgs.bCSV)
-			{
-				CurrentOutputDevice->Logf(TEXT(",Asset,Dependency,Unique KiB,Shared KiB"));
-			}
-			else
-			{
-				CurrentOutputDevice->Logf(TEXT(" %100s %100s %20s %20s"),
-					TEXT("Asset"),
-					TEXT("Dependency"),
-					TEXT("Unique KiB"),
-					TEXT("Shared KiB")
-				);
-			}
-
-			for (TPair<FName, FAssetMemoryDetails>& Asset : AssetMemoryMap)
-			{
-				for (TPair<FName, FAssetMemoryBreakdown>& Dep : Asset.Value.DependenciesToMemoryMap)
-				{
-					const FString DependencyAssetName = Dep.Key.ToString();
-					const FAssetMemoryBreakdown& DepedencyMemoryDetails = Dep.Value;
-
-					if (CommonArgs.bCSV)
-					{
-						CurrentOutputDevice->Logf(TEXT(",%s,%s,%.2f,%.2f"), *Asset.Key.ToString(), *DependencyAssetName,
-							DepedencyMemoryDetails.UniqueSize / 1024.f, DepedencyMemoryDetails.SharedSize / 1024.f);
-					}
-					else
-					{
-						CurrentOutputDevice->Logf(TEXT(" %100s %100s %20.2f %20.2f"),
-							*Asset.Key.ToString(),
-							*DependencyAssetName,
-							DepedencyMemoryDetails.UniqueSize / 1024.f,
-							DepedencyMemoryDetails.SharedSize / 1024.f
-						);
-					}
-				}
-			}
-		}
-
-		if (CommonArgs.bCSV)
-		{
-			CurrentOutputDevice->Logf(TEXT(",TOTAL KiB,%.2f"), TotalCollectionSize / 1024.f);
-		}
-		else
-		{
-			CurrentOutputDevice->Logf(TEXT("TOTAL KiB: %.2f"), TotalCollectionSize / 1024.f);
-		}
-	}
-}));
-
-#endif
-
-namespace MemoryUsageQueries
-{
-
-IMemoryUsageInfoProvider* GetCurrentMemoryUsageInfoProvider()
-{
-	return CurrentMemoryUsageInfoProvider;
-}
-
-bool GetMemoryUsage(const IMemoryUsageInfoProvider* MemoryUsageInfoProvider, const FString& PackageName, uint64& OutExclusiveSize, uint64& OutInclusiveSize, FOutputDevice* ErrorOutput /* = GLog */)
-{
-	FName LongPackageName;
-	TSet<FName> Dependencies;
-	if (!Internal::GetLongNameAndDependencies(PackageName, LongPackageName, Dependencies, ErrorOutput))
-	{
 		return false;
 	}
-
-	OutExclusiveSize = MemoryUsageInfoProvider->GetAssetMemoryUsage(LongPackageName, ErrorOutput);
-	OutInclusiveSize = MemoryUsageInfoProvider->GetAssetsMemoryUsage(Dependencies, ErrorOutput);
-
-	return true;
-}
-
-bool GetMemoryUsageCombined(const IMemoryUsageInfoProvider* MemoryUsageInfoProvider, const TArray<FString>& PackageNames, uint64& OutTotalSize, FOutputDevice* ErrorOutput /* = GLog */)
-{
-	TSet<FName> Dependencies;
-	if (!Internal::GetDependenciesCombined(PackageNames, Dependencies, ErrorOutput))
+	FName LongPackageName;
+	if (!Internal::GetLongName(PackageName, LongPackageName, ErrorOutput))
 	{
 		return false;
 	}
 
-	OutTotalSize =  MemoryUsageInfoProvider->GetAssetsMemoryUsage(Dependencies, ErrorOutput);
-
+	GetMemoryUsage(MemoryUsageInfoProvider, LongPackageName, OutExclusiveSize, OutInclusiveSize);
 	return true;
 }
 
-bool GetMemoryUsageShared(const IMemoryUsageInfoProvider* MemoryUsageInfoProvider, const TArray<FString>& PackageNames, uint64& OutTotalSize, FOutputDevice* ErrorOutput /* = GLog */)
+void GetMemoryUsageCombined(const IMemoryUsageInfoProvider* MemoryUsageInfoProvider, const TArray<FName>& PackageNames, uint64& OutTotalSize)
 {
 	TSet<FName> Dependencies;
-	if (!Internal::GetDependenciesShared(PackageNames, Dependencies, ErrorOutput))
-	{
-		return false;
-	}
-
-	OutTotalSize = MemoryUsageInfoProvider->GetAssetsMemoryUsage(Dependencies, ErrorOutput);
-
-	return true;
-}
-
-bool GetMemoryUsageUnique(const IMemoryUsageInfoProvider* MemoryUsageInfoProvider, const TArray<FString>& PackageNames, uint64& OutUniqueSize, FOutputDevice* ErrorOutput /* = GLog */)
-{
-	TSet<FName> RemovablePackages;
-	if (!Internal::GetRemovablePackages(PackageNames, RemovablePackages, ErrorOutput))
-	{
-		return false;
-	}
-
-	OutUniqueSize = MemoryUsageInfoProvider->GetAssetsMemoryUsage(RemovablePackages, ErrorOutput);
-
-	return true;
+	Internal::GetDependenciesCombined(PackageNames, Dependencies);
+	OutTotalSize =  MemoryUsageInfoProvider->GetAssetsMemoryUsage(Dependencies);
 }
-
-bool GetMemoryUsageCommon(const IMemoryUsageInfoProvider* MemoryUsageInfoProvider, const TArray<FString>& PackageNames, uint64& OutCommonSize, FOutputDevice* ErrorOutput /* = GLog */)
-{
-	TSet<FName> UnremovablePackages;
-	if (!Internal::GetUnremovablePackages(PackageNames, UnremovablePackages, ErrorOutput))
-	{
-		return false;
-	}
-
-	OutCommonSize = MemoryUsageInfoProvider->GetAssetsMemoryUsage(UnremovablePackages, ErrorOutput);
-
-	return true;
-}
-
-bool GatherDependenciesForPackages(const IMemoryUsageInfoProvider* MemoryUsageInfoProvider, const TArray<FString>& PackageNames,
-	TMap<FName, uint64>& OutInternalDeps, TMap<FName, uint64>& OutExternalDeps, MemoryUsageQueries::EDependencyType DependencyType, FOutputDevice* ErrorOutput)
+bool GetMemoryUsageCombined(const IMemoryUsageInfoProvider* MemoryUsageInfoProvider, const TArray<FString>& PackageNames, uint64& OutTotalSize, FOutputDevice* ErrorOutput /* = GLog */)
 {
-	const bool bFindRemovable = (DependencyType == MemoryUsageQueries::EDependencyType::EDep_All || DependencyType == MemoryUsageQueries::EDependencyType::EDep_Removable);
-	const bool bFindNonRemovable = (DependencyType == MemoryUsageQueries::EDependencyType::EDep_All || DependencyType == MemoryUsageQueries::EDependencyType::EDep_NonRemovable);
-
-	TSet<FName> RemovablePackages;
-	if (bFindRemovable && !Internal::GetRemovablePackages(PackageNames, RemovablePackages, ErrorOutput))
-	{
-		return false;
-	}
-
-	TSet<FName> UnremovablePackages;
-	if (bFindNonRemovable && !Internal::GetUnremovablePackages(PackageNames, UnremovablePackages, ErrorOutput))
+	if (!MemoryUsageInfoProvider->IsProviderAvailable())
 	{
+		if (ErrorOutput)
+		{
+			ErrorOutput->Logf(ELogVerbosity::Error, TEXT("MemoryUsageInfoProvider Error: Provider is not available. Please run with -LLM"));
+		}
 		return false;
-	}
-
-	if (bFindRemovable)
-	{
-		MemoryUsageInfoProvider->GetAssetsMemoryUsageWithSize(RemovablePackages, OutInternalDeps, ErrorOutput);
-	}
-
-	if (bFindNonRemovable)
-	{
-		MemoryUsageInfoProvider->GetAssetsMemoryUsageWithSize(UnremovablePackages, OutExternalDeps, ErrorOutput);
-	}
-
-	return true;
-}
-
-
-bool GetDependenciesWithSize(const IMemoryUsageInfoProvider* MemoryUsageInfoProvider, const FString& PackageName, TMap<FName, uint64>& OutDependenciesWithSize, FOutputDevice* ErrorOutput /* = GLog */)
-{
-	FName LongPackageName;
-	TSet<FName> Dependencies;
-
-	if (!Internal::GetLongNameAndDependencies(PackageName, LongPackageName, Dependencies, ErrorOutput))
+	}
+	TSet<FName> LongPackageNames;
+	if (!Internal::GetLongNames(PackageNames, LongPackageNames, ErrorOutput))
 	{
 		return false;
 	}
 
-	OutDependenciesWithSize.Empty();
-	Internal::SortPackagesBySize(MemoryUsageInfoProvider, Dependencies, OutDependenciesWithSize, ErrorOutput);
-
+	GetMemoryUsageCombined(MemoryUsageInfoProvider, LongPackageNames.Array(), OutTotalSize);
 	return true;
 }
 
-bool GetDependenciesWithSizeCombined(const IMemoryUsageInfoProvider* MemoryUsageInfoProvider, const TArray<FString>& PackageNames, TMap<FName, uint64>& OutDependenciesWithSize, FOutputDevice* ErrorOutput /* = GLog */)
+void GetMemoryUsageShared(const IMemoryUsageInfoProvider* MemoryUsageInfoProvider, const TArray<FName>& PackageNames, uint64& OutTotalSize)
 {
 	TSet<FName> Dependencies;
-
-	if (!Internal::GetDependenciesCombined(PackageNames, Dependencies, ErrorOutput))
+	Internal::GetDependenciesShared(PackageNames, Dependencies);
+	OutTotalSize = MemoryUsageInfoProvider->GetAssetsMemoryUsage(Dependencies);
+}
+bool GetMemoryUsageShared(const IMemoryUsageInfoProvider* MemoryUsageInfoProvider, const TArray<FString>& PackageNames, uint64& OutTotalSize, FOutputDevice* ErrorOutput /* = GLog */)
+{
+	if (!MemoryUsageInfoProvider->IsProviderAvailable())
 	{
+		if (ErrorOutput)
+		{
+			ErrorOutput->Logf(ELogVerbosity::Error, TEXT("MemoryUsageInfoProvider Error: Provider is not available. Please run with -LLM"));
+		}
 		return false;
 	}
-
-	OutDependenciesWithSize.Empty();
-	Internal::SortPackagesBySize(MemoryUsageInfoProvider, Dependencies, OutDependenciesWithSize, ErrorOutput);
-
-	return true;
-}
-
-bool GetDependenciesWithSizeShared(const IMemoryUsageInfoProvider* MemoryUsageInfoProvider, const TArray<FString>& PackageNames, TMap<FName, uint64>& OutDependenciesWithSize, FOutputDevice* ErrorOutput /* = GLog */)
-{
-	TSet<FName> Dependencies;
-	if (!Internal::GetDependenciesShared(PackageNames, Dependencies, ErrorOutput))
+	TSet<FName> LongPackageNames;
+	if (!Internal::GetLongNames(PackageNames, LongPackageNames, ErrorOutput))
 	{
 		return false;
 	}
 
-	OutDependenciesWithSize.Empty();
-	Internal::SortPackagesBySize(MemoryUsageInfoProvider, Dependencies, OutDependenciesWithSize, ErrorOutput);
-
+	GetMemoryUsageShared(MemoryUsageInfoProvider, LongPackageNames.Array(), OutTotalSize);
 	return true;
 }
 
-bool GetDependenciesWithSizeUnique(const IMemoryUsageInfoProvider* MemoryUsageInfoProvider, const TArray<FString>& PackageNames, TMap<FName, uint64>& OutDependenciesWithSize, FOutputDevice* ErrorOutput /* = GLog */)
+void GetMemoryUsageUnique(const IMemoryUsageInfoProvider* MemoryUsageInfoProvider, const TArray<FName>& PackageNames, uint64& OutUniqueSize)
 {
 	TSet<FName> RemovablePackages;
-	if (!Internal::GetRemovablePackages(PackageNames, RemovablePackages, ErrorOutput))
+	Internal::GetRemovablePackages(PackageNames, RemovablePackages);
+	OutUniqueSize = MemoryUsageInfoProvider->GetAssetsMemoryUsage(RemovablePackages);
+}
+bool GetMemoryUsageUnique(const IMemoryUsageInfoProvider* MemoryUsageInfoProvider, const TArray<FString>& PackageNames, uint64& OutUniqueSize, FOutputDevice* ErrorOutput /* = GLog */)
+{
+	if (!MemoryUsageInfoProvider->IsProviderAvailable())
+	{
+		if (ErrorOutput)
+		{
+			ErrorOutput->Logf(ELogVerbosity::Error, TEXT("MemoryUsageInfoProvider Error: Provider is not available. Please run with -LLM"));
+		}
+		return false;
+	}
+	TSet<FName> LongPackageNames;
+	if (!Internal::GetLongNames(PackageNames, LongPackageNames, ErrorOutput))
 	{
 		return false;
 	}
 
-	OutDependenciesWithSize.Empty();
-	Internal::SortPackagesBySize(MemoryUsageInfoProvider, RemovablePackages, OutDependenciesWithSize, ErrorOutput);
-
+	GetMemoryUsageUnique(MemoryUsageInfoProvider, LongPackageNames.Array(), OutUniqueSize);
 	return true;
 }
 
-bool GetDependenciesWithSizeCommon(const IMemoryUsageInfoProvider* MemoryUsageInfoProvider, const TArray<FString>& PackageNames, TMap<FName, uint64>& OutDependenciesWithSize, FOutputDevice* ErrorOutput /* = GLog */)
+void GetMemoryUsageCommon(const IMemoryUsageInfoProvider* MemoryUsageInfoProvider, const TArray<FName>& PackageNames, uint64& OutCommonSize)
 {
 	TSet<FName> UnremovablePackages;
-	if (!Internal::GetUnremovablePackages(PackageNames, UnremovablePackages, ErrorOutput))
+	Internal::GetUnremovablePackages(PackageNames, UnremovablePackages);
+	OutCommonSize = MemoryUsageInfoProvider->GetAssetsMemoryUsage(UnremovablePackages);
+}
+bool GetMemoryUsageCommon(const IMemoryUsageInfoProvider* MemoryUsageInfoProvider, const TArray<FString>& PackageNames, uint64& OutCommonSize, FOutputDevice* ErrorOutput /* = GLog */)
+{
+	if (!MemoryUsageInfoProvider->IsProviderAvailable())
+	{
+		if (ErrorOutput)
+		{
+			ErrorOutput->Logf(ELogVerbosity::Error, TEXT("MemoryUsageInfoProvider Error: Provider is not available. Please run with -LLM"));
+		}
+		return false;
+	}
+	TSet<FName> LongPackageNames;
+	if (!Internal::GetLongNames(PackageNames, LongPackageNames, ErrorOutput))
 	{
 		return false;
 	}
 
-	OutDependenciesWithSize.Empty();
-	Internal::SortPackagesBySize(MemoryUsageInfoProvider, UnremovablePackages, OutDependenciesWithSize, ErrorOutput);
-
+	GetMemoryUsageCommon(MemoryUsageInfoProvider, LongPackageNames.Array(), OutCommonSize);
 	return true;
 }
 
-#if ENABLE_LOW_LEVEL_MEM_TRACKER
-bool GetFilteredPackagesWithSize(TMap<FName, uint64>& OutPackagesWithSize, FName GroupName /* = NAME_None */, FString AssetSubstring /* = FString() */, FName ClassName /* = NAME_None */, FOutputDevice* ErrorOutput /* = GLog */)
+void GatherDependenciesForPackages(const IMemoryUsageInfoProvider* MemoryUsageInfoProvider, const TArray<FName>& PackageNames,
+	TMap<FName, uint64>& OutInternalDeps, TMap<FName, uint64>& OutExternalDeps, MemoryUsageQueries::EDependencyType DependencyType)
 {
-	TArray<FLLMTagSetAllocationFilter> Filters;
-	if (GroupName != NAME_None)
+	const bool bFindRemovable = (DependencyType == MemoryUsageQueries::EDependencyType::EDep_All || DependencyType == MemoryUsageQueries::EDependencyType::EDep_Removable);
+	const bool bFindNonRemovable = (DependencyType == MemoryUsageQueries::EDependencyType::EDep_All || DependencyType == MemoryUsageQueries::EDependencyType::EDep_NonRemovable);
+
+	if (bFindRemovable)
 	{
-		Filters.Add({GroupName, ELLMTagSet::None});
+		TSet<FName> RemovablePackages;
+		Internal::GetRemovablePackages(PackageNames, RemovablePackages);
+		MemoryUsageInfoProvider->GetAssetsMemoryUsageWithSize(RemovablePackages, OutInternalDeps);
 	}
-
-	if (ClassName != NAME_None)
+	if (bFindNonRemovable)
 	{
-		Filters.Add({ClassName, ELLMTagSet::AssetClasses});
+		TSet<FName> UnremovablePackages;
+		Internal::GetUnremovablePackages(PackageNames, UnremovablePackages);
+		MemoryUsageInfoProvider->GetAssetsMemoryUsageWithSize(UnremovablePackages, OutExternalDeps);
 	}
 
-	MemoryUsageInfoProviderLLM.GetFilteredTagsWithSize(OutPackagesWithSize, ELLMTracker::Default, ELLMTagSet::Assets, Filters, ErrorOutput);
-
-	if (AssetSubstring.Len() > 0)
+}
+bool GatherDependenciesForPackages(const IMemoryUsageInfoProvider* MemoryUsageInfoProvider, const TArray<FString>& PackageNames,
+	TMap<FName, uint64>& OutInternalDeps, TMap<FName, uint64>& OutExternalDeps, MemoryUsageQueries::EDependencyType DependencyType, FOutputDevice* ErrorOutput)
+{
+	if (!MemoryUsageInfoProvider->IsProviderAvailable())
 	{
-		Internal::RemoveFilteredPackages(OutPackagesWithSize, AssetSubstring);
+		if (ErrorOutput)
+		{
+			ErrorOutput->Logf(ELogVerbosity::Error, TEXT("MemoryUsageInfoProvider Error: Provider is not available. Please run with -LLM"));
+		}
+		return false;
 	}
-
-	OutPackagesWithSize.ValueSort(TGreater<uint64>());
-
+	TSet<FName> LongPackageNames;
+	if (!Internal::GetLongNames(PackageNames, LongPackageNames, ErrorOutput))
+	{
+		return false;
+	}
+	GatherDependenciesForPackages(MemoryUsageInfoProvider, LongPackageNames.Array(), OutInternalDeps, OutExternalDeps, DependencyType);
 	return true;
 }
 
-bool GetFilteredPackagesCategorizedWithSize(TMap<FName, uint64>& OutPackagesWithSize, FName GroupName /*= NAME_None*/, FString AssetSubstring /*= FString()*/, FName ClassName /*= NAME_None*/, FName CategoryName /*= NAME_None*/, FOutputDevice* ErrorOutput /*= GLog*/)
+void GetDependenciesWithSize(const IMemoryUsageInfoProvider* MemoryUsageInfoProvider, const FName& PackageName, TMap<FName, uint64>& OutDependenciesWithSize)
 {
-	TArray<FLLMTagSetAllocationFilter> Filters;
-
-	FName LongName = NAME_None;
-	if (!AssetSubstring.IsEmpty())
+	TSet<FName> Dependencies;
+	Internal::GetTransitiveDependencies(PackageName, Dependencies);
+	Dependencies.Add(PackageName);
+	OutDependenciesWithSize.Empty();
+	Internal::SortPackagesBySize(MemoryUsageInfoProvider, Dependencies, OutDependenciesWithSize);
+}
+bool GetDependenciesWithSize(const IMemoryUsageInfoProvider* MemoryUsageInfoProvider, const FString& PackageName, TMap<FName, uint64>& OutDependenciesWithSize, FOutputDevice* ErrorOutput /* = GLog */)
+{
+	if (!MemoryUsageInfoProvider->IsProviderAvailable())
 	{
-		// Special handling for "Untagged" assets as llm uses this tag to bucket allocations without asset scopes.
-		FName UntaggedAssetTagName = LLMGetUntaggedTagName(ELLMTagSet::Assets);
-		if (FName(AssetSubstring) == UntaggedAssetTagName)
-		{
-			LongName = UntaggedAssetTagName;
-		}
-		else
+		if (ErrorOutput)
 		{
-			if (!Internal::GetLongName(AssetSubstring, LongName, ErrorOutput))
-			{
-				return false;
-			}
+			ErrorOutput->Logf(ELogVerbosity::Error, TEXT("MemoryUsageInfoProvider Error: Provider is not available. Please run with -LLM"));
 		}
+		return false;
 	}
-
-	if (LongName != NAME_None)
+	FName LongPackageName;
+	if (!Internal::GetLongName(PackageName, LongPackageName, ErrorOutput))
 	{
-		Filters.Add({ LongName, ELLMTagSet::Assets });
+		return false;
 	}
 
-	if (GroupName != NAME_None)
+	GetDependenciesWithSize(MemoryUsageInfoProvider, LongPackageName, OutDependenciesWithSize);
+	return true;
+}
+
+void GetDependenciesWithSizeCombined(const IMemoryUsageInfoProvider* MemoryUsageInfoProvider, const TArray<FName>& PackageNames, TMap<FName, uint64>& OutDependenciesWithSize)
+{
+	TSet<FName> Dependencies;
+	Internal::GetDependenciesCombined(PackageNames, Dependencies);
+	OutDependenciesWithSize.Empty();
+	Internal::SortPackagesBySize(MemoryUsageInfoProvider, Dependencies, OutDependenciesWithSize);
+}
+bool GetDependenciesWithSizeCombined(const IMemoryUsageInfoProvider* MemoryUsageInfoProvider, const TArray<FString>& PackageNames, TMap<FName, uint64>& OutDependenciesWithSize, FOutputDevice* ErrorOutput /* = GLog */)
+{
+	if (!MemoryUsageInfoProvider->IsProviderAvailable())
 	{
-		Filters.Add({ GroupName, ELLMTagSet::None });
+		if (ErrorOutput)
+		{
+			ErrorOutput->Logf(ELogVerbosity::Error, TEXT("MemoryUsageInfoProvider Error: Provider is not available. Please run with -LLM"));
+		}
+		return false;
 	}
-
-	if (ClassName != NAME_None)
+	TSet<FName> LongPackageNames;
+	if (!Internal::GetLongNames(PackageNames, LongPackageNames, ErrorOutput))
 	{
-		Filters.Add({ ClassName, ELLMTagSet::AssetClasses });
+		return false;
 	}
 
-	ELLMTagSet TagSetCategory = ELLMTagSet::None;
-	if (CategoryName == FName("Assets"))
+	GetDependenciesWithSizeCombined(MemoryUsageInfoProvider, LongPackageNames.Array(), OutDependenciesWithSize);
+	return true;
+}
+void GetDependenciesWithSizeShared(const IMemoryUsageInfoProvider* MemoryUsageInfoProvider, const TArray<FName>& PackageNames, TMap<FName, uint64>& OutDependenciesWithSize)
+{
+	TSet<FName> Dependencies;
+	Internal::GetDependenciesShared(PackageNames, Dependencies);
+	OutDependenciesWithSize.Empty();
+	Internal::SortPackagesBySize(MemoryUsageInfoProvider, Dependencies, OutDependenciesWithSize);
+}
+bool GetDependenciesWithSizeShared(const IMemoryUsageInfoProvider* MemoryUsageInfoProvider, const TArray<FString>& PackageNames, TMap<FName, uint64>& OutDependenciesWithSize, FOutputDevice* ErrorOutput /* = GLog */)
+{
+	if (!MemoryUsageInfoProvider->IsProviderAvailable())
 	{
-		TagSetCategory = ELLMTagSet::Assets;
+		if (ErrorOutput)
+		{
+			ErrorOutput->Logf(ELogVerbosity::Error, TEXT("MemoryUsageInfoProvider Error: Provider is not available. Please run with -LLM"));
+		}
+		return false;
 	}
-	else  if (CategoryName == FName("AssetClasses"))
+	TSet<FName> LongPackageNames;
+	if (!Internal::GetLongNames(PackageNames, LongPackageNames, ErrorOutput))
 	{
-		TagSetCategory = ELLMTagSet::AssetClasses;
+		return false;
 	}
 
-	MemoryUsageInfoProviderLLM.GetFilteredTagsWithSize(OutPackagesWithSize, ELLMTracker::Default, TagSetCategory, Filters, ErrorOutput);
-
-	OutPackagesWithSize.ValueSort(TGreater<uint64>());
-
+	GetDependenciesWithSizeShared(MemoryUsageInfoProvider, LongPackageNames.Array(), OutDependenciesWithSize);
 	return true;
 }
 
-bool GetFilteredClassesWithSize(TMap<FName, uint64>& OutClassesWithSize, FName GroupName /* = NAME_None */, FString AssetName /* = FString() */, FOutputDevice* ErrorOutput /* = GLog */)
+void GetDependenciesWithSizeUnique(const IMemoryUsageInfoProvider* MemoryUsageInfoProvider, const TArray<FName>& PackageNames, TMap<FName, uint64>& OutDependenciesWithSize)
 {
-	TArray<FLLMTagSetAllocationFilter> Filters;
-
-	FName LongName = NAME_None;
-	if (!AssetName.IsEmpty() && !Internal::GetLongName(AssetName, LongName, ErrorOutput))
+	TSet<FName> RemovablePackages;
+	Internal::GetRemovablePackages(PackageNames, RemovablePackages);
+	OutDependenciesWithSize.Empty();
+	Internal::SortPackagesBySize(MemoryUsageInfoProvider, RemovablePackages, OutDependenciesWithSize);
+}
+bool GetDependenciesWithSizeUnique(const IMemoryUsageInfoProvider* MemoryUsageInfoProvider, const TArray<FString>& PackageNames, TMap<FName, uint64>& OutDependenciesWithSize, FOutputDevice* ErrorOutput /* = GLog */)
+{
+	if (!MemoryUsageInfoProvider->IsProviderAvailable())
 	{
+		if (ErrorOutput)
+		{
+			ErrorOutput->Logf(ELogVerbosity::Error, TEXT("MemoryUsageInfoProvider Error: Provider is not available. Please run with -LLM"));
+		}
 		return false;
 	}
-
-	if (LongName != NAME_None)
-	{
-		Filters.Add({LongName, ELLMTagSet::Assets});
-	}
-
-	if (GroupName != NAME_None)
+	TSet<FName> LongPackageNames;
+	if (!Internal::GetLongNames(PackageNames, LongPackageNames, ErrorOutput))
 	{
-		Filters.Add({GroupName, ELLMTagSet::None});
+		return false;
 	}
 
-	MemoryUsageInfoProviderLLM.GetFilteredTagsWithSize(OutClassesWithSize, ELLMTracker::Default, ELLMTagSet::AssetClasses, Filters, ErrorOutput);
-
-	OutClassesWithSize.ValueSort(TGreater<uint64>());
-
+	GetDependenciesWithSizeUnique(MemoryUsageInfoProvider, LongPackageNames.Array(), OutDependenciesWithSize);
 	return true;
 }
 
-bool GetFilteredGroupsWithSize(TMap<FName, uint64>& OutGroupsWithSize, FString AssetName /* = FString() */, FName ClassName /* = NAME_None */, FOutputDevice* ErrorOutput /* = GLog */)
+void GetDependenciesWithSizeCommon(const IMemoryUsageInfoProvider* MemoryUsageInfoProvider, const TArray<FName>& PackageNames, TMap<FName, uint64>& OutDependenciesWithSize)
 {
-	TArray<FLLMTagSetAllocationFilter> Filters;
-
-	FName LongName = NAME_None;
-	if (!AssetName.IsEmpty() && !Internal::GetLongName(AssetName, LongName, ErrorOutput))
+	TSet<FName> UnremovablePackages;
+	Internal::GetUnremovablePackages(PackageNames, UnremovablePackages);
+	OutDependenciesWithSize.Empty();
+	Internal::SortPackagesBySize(MemoryUsageInfoProvider, UnremovablePackages, OutDependenciesWithSize);
+}
+bool GetDependenciesWithSizeCommon(const IMemoryUsageInfoProvider* MemoryUsageInfoProvider, const TArray<FString>& PackageNames, TMap<FName, uint64>& OutDependenciesWithSize, FOutputDevice* ErrorOutput /* = GLog */)
+{
+	if (!MemoryUsageInfoProvider->IsProviderAvailable())
 	{
+		if (ErrorOutput)
+		{
+			ErrorOutput->Logf(ELogVerbosity::Error, TEXT("MemoryUsageInfoProvider Error: Provider is not available. Please run with -LLM"));
+		}
 		return false;
 	}
-
-	if (LongName != NAME_None)
-	{
-		Filters.Add({LongName, ELLMTagSet::Assets});
-	}
-
-	if (ClassName != NAME_None)
+	TSet<FName> LongPackageNames;
+	if (!Internal::GetLongNames(PackageNames, LongPackageNames, ErrorOutput))
 	{
-		Filters.Add({ClassName, ELLMTagSet::AssetClasses});
+		return false;
 	}
 
-	MemoryUsageInfoProviderLLM.GetFilteredTagsWithSize(OutGroupsWithSize, ELLMTracker::Default, ELLMTagSet::None, Filters, ErrorOutput);
-
-	OutGroupsWithSize.ValueSort(TGreater<uint64>());
-
+	GetDependenciesWithSizeCommon(MemoryUsageInfoProvider, LongPackageNames.Array(), OutDependenciesWithSize);
 	return true;
 }
-#endif
 
 namespace Internal
 {
@@ -1173,13 +336,8 @@ FMemoryUsageReferenceProcessor::FMemoryUsageReferenceProcessor()
 {
 	int32 Num = GUObjectArray.GetObjectArrayNum();
 
-	Excluded.Init(false, Num);
 	ReachableFull.Init(false, Num);
-	ReachableExcluded.Init(false, Num);
-}
 
-bool FMemoryUsageReferenceProcessor::Init(const TArray<FString>& PackageNames, FOutputDevice* ErrorOutput /* = GLog */)
-{
 	for (FRawObjectIterator It(true); It; ++It)
 	{
 		FUObjectItem* ObjectItem = *It;
@@ -1209,22 +367,25 @@ bool FMemoryUsageReferenceProcessor::Init(const TArray<FString>& PackageNames, F
 	{
 		RootSetPackages.Add(FGCObject::GGCObjectReferencer);
 	}
+}
 
-	TSet<FName> LongPackageNames;
-	if (!GetLongNames(PackageNames, LongPackageNames, ErrorOutput))
-	{
-		return false;
-	}
+void FMemoryUsageReferenceProcessor::AddAssetSet(TConstArrayView<FName> PackageNames)
+{
+	FAssetSet& AssetSet = AssetSets.Emplace_GetRef();
 
-	for (const auto& PackageName : LongPackageNames)
+	int Num = ReachableFull.Num();
+	AssetSet.Excluded.Init(false, Num);
+	AssetSet.ReachableExcluded.Init(false, Num);
+	
+	for (const FName& PackageName : PackageNames)
 	{
 		UPackage* Package = Cast<UPackage>(StaticFindObjectFast(UPackage::StaticClass(), nullptr, PackageName, /*bExactClass*/ true));
 		if (Package)
 		{
 			auto ExcludeObject = [&](UObject* ObjectToExclude) {
-				if (ObjectToExclude && GUObjectArray.ObjectToIndex(ObjectToExclude) < Excluded.Num())
+				if (ObjectToExclude && GUObjectArray.ObjectToIndex(ObjectToExclude) < Num)
 				{
-					Excluded[GUObjectArray.ObjectToIndex(ObjectToExclude)] = true;
+					AssetSet.Excluded[GUObjectArray.ObjectToIndex(ObjectToExclude)] = true;
 				}
 			};
 
@@ -1249,26 +410,24 @@ bool FMemoryUsageReferenceProcessor::Init(const TArray<FString>& PackageNames, F
 			ExcludeObject(Package);
 		}
 	}
-
-	return true;
-}
-
-TArray<UObject*>& FMemoryUsageReferenceProcessor::GetRootSet()
-{
-	return RootSetPackages;
 }
 
 void FMemoryUsageReferenceProcessor::HandleTokenStreamObjectReference(UE::GC::FWorkerContext& Context, const UObject* ReferencingObject, UObject*& Object, UE::GC::FTokenId TokenIndex, const EGCTokenType TokenType, bool bAllowReferenceElimination)
 {
 	FPermanentObjectPoolExtents PermanentPool;
-	if (Object == nullptr || GUObjectArray.ObjectToIndex(Object) >= ReachableFull.Num() || PermanentPool.Contains(Object) || GUObjectArray.IsDisregardForGC(Object))
+	if (Object == nullptr || PermanentPool.Contains(Object) || GUObjectArray.IsDisregardForGC(Object))
 	{
 		return;
 	}
 
 	int32 ObjectIndex = GUObjectArray.ObjectToIndex(Object);
 
-	if (Mode == EMode::Full)
+	if (ObjectIndex >= ReachableFull.Num())
+	{
+		return;
+	}
+
+	if (bFirstPass)
 	{
 		if (!ReachableFull[ObjectIndex])
 		{
@@ -1276,21 +435,33 @@ void FMemoryUsageReferenceProcessor::HandleTokenStreamObjectReference(UE::GC::FW
 			Context.ObjectsToSerialize.Add<Options>(Object);
 		}
 	}
-	else if (Mode == EMode::Excluding)
+	else
 	{
-		if (!ReachableExcluded[ObjectIndex] && !Excluded[ObjectIndex])
+		FAssetSet& AssetSet = AssetSets[PassesLeft - 1];
+		if (!AssetSet.ReachableExcluded[ObjectIndex] && !AssetSet.Excluded[ObjectIndex])
 		{
-			ReachableExcluded[ObjectIndex] = true;
+			AssetSet.ReachableExcluded[ObjectIndex] = true;
 			Context.ObjectsToSerialize.Add<Options>(Object);
 		}
 	}
 }
 
-bool FMemoryUsageReferenceProcessor::GetUnreachablePackages(TSet<FName>& OutUnreachablePackages)
+void FMemoryUsageReferenceProcessor::PerformReachabilityAnalysis()
+{
+	UE::GC::PadObjectArray(RootSetPackages);
+	for(;NeedsAnotherPass(); ProgressToNextPass())
+	{
+		FGCArrayStruct ArrayStruct;
+		ArrayStruct.SetInitialObjectsPrepadded(RootSetPackages);
+		CollectReferences<FMemoryUsageReferenceCollector, FMemoryUsageReferenceProcessor>(*this, ArrayStruct);
+	}
+}
+
+void FMemoryUsageReferenceProcessor::GetUnreachablePackages(int AssetSetIdx, TSet<FName>& OutUnreachablePackages) const
 {
 	for (int i = 0; i < ReachableFull.Num(); i++)
 	{
-		if (ReachableFull[i] && !ReachableExcluded[i])
+		if (ReachableFull[i] && !AssetSets[AssetSetIdx].ReachableExcluded[i])
 		{
 			UObject* Obj = static_cast<UObject*>(GUObjectArray.IndexToObjectUnsafeForGC(i)->Object);
 			if (Obj && Obj->IsA<UPackage>())
@@ -1299,8 +470,37 @@ bool FMemoryUsageReferenceProcessor::GetUnreachablePackages(TSet<FName>& OutUnre
 			}
 		}
 	}
+}
 
-	return true;
+void FMemoryUsageReferenceProcessor::GetUnreachablePackages(TArrayView<FRemovablePackagesInfo> OutUnreachablePackages) const
+{
+	int Num = AssetSets.Num();
+	if(ensure(OutUnreachablePackages.Num() == Num))
+	{
+		for (int i = 0; i < ReachableFull.Num(); i++)
+		{
+			if (!ReachableFull[i])
+			{
+				continue;
+			}
+			UObject* Obj = nullptr;
+			for (int AssetSetIdx = 0; AssetSetIdx < Num; ++AssetSetIdx)
+			{
+				const FAssetSet& AssetSet = AssetSets[AssetSetIdx];
+				if (!AssetSet.ReachableExcluded[i])
+				{
+					if (!Obj)
+					{
+						Obj = static_cast<UObject*>(GUObjectArray.IndexToObjectUnsafeForGC(i)->Object);
+					}
+					if (Obj && Obj->IsA<UPackage>())
+					{
+						OutUnreachablePackages[AssetSetIdx].RemovablePackages.Add(Obj->GetFName());
+					}
+				}
+			}
+		}
+	}
 }
 
 static FAutoConsoleVariable CVarMemQueryUsePackageStore(
@@ -1720,50 +920,27 @@ bool GetLongNames(const TArray<FString>& PackageNames, TSet<FName>& OutLongPacka
 	return true;
 }
 
-bool GetLongNameAndDependencies(FStringView PackageName, FName& OutLongPackageName, TSet<FName>& OutDependencies, FOutputDevice* ErrorOutput /* = GLog */)
-{
-	if (!GetLongName(PackageName, OutLongPackageName, ErrorOutput))
-	{
-		return false;
-	}
-
-	GetTransitiveDependencies(OutLongPackageName, OutDependencies);
-	OutDependencies.Add(OutLongPackageName);
-
-	return true;
-}
-
-bool GetDependenciesCombined(const TArray<FString>& PackageNames, TSet<FName>& OutDependencies, FOutputDevice* ErrorOutput /* = GLog */)
+void GetDependenciesCombined(const TArray<FName>& PackageNames, TSet<FName>& OutDependencies)
 {
-	FName LongPackageName;
 	TSet<FName> Dependencies;
-
-	for (const auto& Package : PackageNames)
+	for (const auto& PackageName : PackageNames)
 	{
 		Dependencies.Empty();
-		if (!Internal::GetLongNameAndDependencies(Package, LongPackageName, Dependencies, ErrorOutput))
-		{
-			return false;
-		}
-
+		Internal::GetTransitiveDependencies(PackageName, Dependencies);
 		OutDependencies.Append(Dependencies);
 	}
-
-	return true;
+	OutDependencies.Append(PackageNames);
 }
 
-bool GetDependenciesShared(const TArray<FString>& PackageNames, TSet<FName>& OutDependencies, FOutputDevice* ErrorOutput /* = GLog */)
+void GetDependenciesShared(const TArray<FName>& PackageNames, TSet<FName>& OutDependencies)
 {
-	FName LongPackageName;
 	TSet<FName> Dependencies;
 
 	for (int i = 0; i < PackageNames.Num(); i++)
 	{
 		Dependencies.Empty();
-		if (!Internal::GetLongNameAndDependencies(PackageNames[i], LongPackageName, Dependencies, ErrorOutput))
-		{
-			return false;
-		}
+		Internal::GetTransitiveDependencies(PackageNames[i], Dependencies);
+		Dependencies.Add(PackageNames[i]);
 
 		if (i == 0)
 		{
@@ -1773,71 +950,27 @@ bool GetDependenciesShared(const TArray<FString>& PackageNames, TSet<FName>& Out
 
 		OutDependencies = OutDependencies.Intersect(Dependencies);
 	}
-
-	return true;
-}
-
-bool PerformReachabilityAnalysis(FMemoryUsageReferenceProcessor& ReferenceProcessor, FOutputDevice* ErrorOutput /* = GLog */)
-{
-	{
-		FGCArrayStruct ArrayStruct;
-		ArrayStruct.SetInitialObjectsUnpadded(ReferenceProcessor.GetRootSet());
-		ReferenceProcessor.SetMode(FMemoryUsageReferenceProcessor::Full);
-		CollectReferences<FMemoryUsageReferenceCollector, FMemoryUsageReferenceProcessor>(ReferenceProcessor, ArrayStruct);
-	}
-
-	{
-		FGCArrayStruct ArrayStruct;
-		ArrayStruct.SetInitialObjectsUnpadded(ReferenceProcessor.GetRootSet());
-		ReferenceProcessor.SetMode(FMemoryUsageReferenceProcessor::Excluding);
-		CollectReferences<FMemoryUsageReferenceCollector, FMemoryUsageReferenceProcessor>(ReferenceProcessor, ArrayStruct);
-	}
-
-	return true;
 }
 
-bool GetRemovablePackages(const TArray<FString>& PackagesToUnload, TSet<FName>& OutRemovablePackages, FOutputDevice* ErrorOutput /* = GLog */)
+void GetRemovablePackages(const TArray<FName>& PackagesToUnload, TSet<FName>& OutRemovablePackages)
 {
 	FMemoryUsageReferenceProcessor ReferenceProcessor;
-
-	if (!ReferenceProcessor.Init(PackagesToUnload, ErrorOutput))
-	{
-		return false;
-	}
-
-	if (!PerformReachabilityAnalysis(ReferenceProcessor, ErrorOutput))
-	{
-		return false;
-	}
-
-	return ReferenceProcessor.GetUnreachablePackages(OutRemovablePackages);
+	ReferenceProcessor.AddAssetSet(PackagesToUnload);
+	ReferenceProcessor.PerformReachabilityAnalysis();
+	ReferenceProcessor.GetUnreachablePackages(0, OutRemovablePackages);
 }
 
-bool GetUnremovablePackages(const TArray<FString>& PackagesToUnload, TSet<FName>& OutUnremovablePackages, FOutputDevice* ErrorOutput /* = GLog */)
+void GetUnremovablePackages(const TArray<FName>& PackagesToUnload, TSet<FName>& OutUnremovablePackages)
 {
 	FMemoryUsageReferenceProcessor ReferenceProcessor;
-
-	if (!ReferenceProcessor.Init(PackagesToUnload, ErrorOutput))
-	{
-		return false;
-	}
-
-	if (!PerformReachabilityAnalysis(ReferenceProcessor, ErrorOutput))
-	{
-		return false;
-	}
+	ReferenceProcessor.AddAssetSet(PackagesToUnload);
+	ReferenceProcessor.PerformReachabilityAnalysis();
 
 	TSet<FName> UnreachablePackages;
-	if (!ReferenceProcessor.GetUnreachablePackages(UnreachablePackages))
-	{
-		return false;
-	}
+	ReferenceProcessor.GetUnreachablePackages(0, UnreachablePackages);
 
 	TSet<FName> Dependencies;
-	if (!Internal::GetDependenciesCombined(PackagesToUnload, Dependencies, ErrorOutput))
-	{
-		return false;
-	}
+	GetDependenciesCombined(PackagesToUnload, Dependencies);
 
 	for (const auto& Package : Dependencies)
 	{
@@ -1846,8 +979,17 @@ bool GetUnremovablePackages(const TArray<FString>& PackagesToUnload, TSet<FName>
 			OutUnremovablePackages.Add(Package);
 		}
 	}
+}
 
-	return true;
+void GetRemovablePackagesBulk(TArrayView<FMemoryUsageReferenceProcessor::FRemovablePackagesInfo> Infos)
+{
+	FMemoryUsageReferenceProcessor ReferenceProcessor;
+	for (const FMemoryUsageReferenceProcessor::FRemovablePackagesInfo& Info : Infos)
+	{
+		ReferenceProcessor.AddAssetSet(Info.Assets);
+	}
+	ReferenceProcessor.PerformReachabilityAnalysis();
+	ReferenceProcessor.GetUnreachablePackages(Infos);
 }
 
 static void GetTransitiveDependenciesAssetRegistry(FName PackageName, TSet<FName>& OutDependencies)
@@ -1926,19 +1068,17 @@ void GetTransitiveDependencies(FName PackageName, TSet<FName>& OutDependencies)
 	}
 }
 
-void SortPackagesBySize(
-	const IMemoryUsageInfoProvider* MemoryUsageInfoProvider, const TSet<FName>& Packages, TMap<FName, uint64>& OutPackagesWithSize,
-	FOutputDevice* ErrorOutput /* = GLog */)
+void SortPackagesBySize(const IMemoryUsageInfoProvider* MemoryUsageInfoProvider, const TSet<FName>& Packages, TMap<FName, uint64>& OutPackagesWithSize)
 {
-	GetPackagesSize(MemoryUsageInfoProvider, Packages, OutPackagesWithSize, ErrorOutput);
+	GetPackagesSize(MemoryUsageInfoProvider, Packages, OutPackagesWithSize);
 	OutPackagesWithSize.ValueSort([&](const uint64& A, const uint64& B) { return A > B; });
 }
 
-void GetPackagesSize(const IMemoryUsageInfoProvider* MemoryUsageInfoProvider, const TSet<FName>& Packages, TMap<FName, uint64>& OutPackagesWithSize, FOutputDevice* ErrorOutput /* = GLog */)
+void GetPackagesSize(const IMemoryUsageInfoProvider* MemoryUsageInfoProvider, const TSet<FName>& Packages, TMap<FName, uint64>& OutPackagesWithSize)
 {
 	for (const auto& Package : Packages)
 	{
-		uint64 Size = MemoryUsageInfoProvider->GetAssetMemoryUsage(Package, ErrorOutput);
+		uint64 Size = MemoryUsageInfoProvider->GetAssetMemoryUsage(Package);
 		OutPackagesWithSize.Add(Package, Size);
 	}
 }
@@ -1991,6 +1131,8 @@ void RemoveFilteredPackages(TMap<FName, uint64>& OutPackagesWithSize, const FStr
 	}
 }
 
+static int32 DefaultResultLimit = 15;
+
 void PrintTagsWithSize(FOutputDevice& Ar, const TMap<FName, uint64>& TagsWithSize, const TCHAR* Name, bool bTruncate /* = false */, int32 Limit /* = -1 */, bool bCSV /* = false */)
 {
 	uint64 TotalSize = 0U;
diff --git a/Engine/Plugins/MemoryUsageQueries/Source/MemoryUsageQueries/Private/MemoryUsageQueriesCommands.cpp b/Engine/Plugins/MemoryUsageQueries/Source/MemoryUsageQueries/Private/MemoryUsageQueriesCommands.cpp
new file mode 100644
--- /dev/null
+++ b/Engine/Plugins/MemoryUsageQueries/Source/MemoryUsageQueries/Private/MemoryUsageQueriesCommands.cpp
@@ -0,0 +1,949 @@
+// Copyright Epic Games, Inc. All Rights Reserved.
+
+#include "MemoryUsageQueries.h"
+#include "MemoryUsageQueriesConfig.h"
+#include "MemoryUsageQueriesPrivate.h"
+#include "Engine/BlueprintGeneratedClass.h"
+#include "Misc/Paths.h"
+#include "Misc/PackageName.h"
+#include "Misc/OutputDeviceArchiveWrapper.h"
+#include "Misc/WildcardString.h"
+#include "HAL/FileManager.h"
+
+struct FAssetMemoryBreakdown
+{
+	FAssetMemoryBreakdown()
+	{
+	}
+
+	uint64 ExclusiveSize = 0;
+	uint64 UniqueSize = 0;
+	uint64 SharedSize = 0;
+	uint64 TotalSize = 0;
+};
+
+struct FAssetMemoryDetails
+{
+	FAssetMemoryDetails()
+	{
+	}
+
+	// Assets Package Name
+	FName PackageName;
+
+	FAssetMemoryBreakdown MemoryBreakdown;
+
+	// List of dependencies for this asset
+	TSet<FName> Dependencies;
+
+	TMap<FName, FAssetMemoryBreakdown> DependenciesToMemoryMap;
+
+	int32 UniqueRefCount = 0;
+	int32 SharedRefCount = 0;
+};
+
+// Helper struct that scopes an output device, and will automatically clean up any log files when leaving scope
+struct FScopedOutputDevice
+{
+	FOutputDevice* CurrentOutputDevice = nullptr;
+
+	FOutputDeviceArchiveWrapper* FileArWrapper = nullptr;
+	FArchive* FileAr = nullptr;
+
+	explicit FScopedOutputDevice(FOutputDevice* InDefaultOutputDevice) :
+		CurrentOutputDevice(InDefaultOutputDevice)
+		
+	{
+	}
+
+	~FScopedOutputDevice()
+	{
+		if (FileArWrapper != nullptr)
+		{
+			FileArWrapper->TearDown();
+		}
+
+		delete FileArWrapper;
+		delete FileAr;
+	}
+
+	FOutputDevice& GetOutputDevice() const
+	{
+		return *CurrentOutputDevice;
+	}
+
+	void OpenLogFile(const FString& LogFileName, bool bCSV)
+	{
+		if (FileArWrapper != nullptr || LogFileName.IsEmpty())
+		{
+			// Log file is already open or LogFileName is empty
+			return;
+		}
+
+#if ALLOW_DEBUG_FILES
+		// Create folder for MemQuery files.
+		const FString OutputDir = FPaths::ProfilingDir() / TEXT("MemQuery");
+		IFileManager::Get().MakeDirectory(*OutputDir, true);
+
+		FString FileTimeString = FString::Printf(TEXT("_%s"), *FDateTime::Now().ToString(TEXT("%H%M%S")));
+		const FString FileExtension = (bCSV ? TEXT(".csv") : TEXT(".memquery"));
+		const FString LogFilename = OutputDir / (LogFileName + FileTimeString + FileExtension);
+
+		FileAr = IFileManager::Get().CreateDebugFileWriter(*LogFilename);
+
+		if (FileAr != nullptr)
+		{
+			FileArWrapper = new FOutputDeviceArchiveWrapper(FileAr);
+				
+			if (FileArWrapper != nullptr)
+			{
+				CurrentOutputDevice = FileArWrapper;
+			}
+			else
+			{
+				delete FileAr;
+				FileAr = nullptr;
+			}
+		}
+#endif
+	}
+
+	// Delete default constructor to force proper usage
+	FScopedOutputDevice() = delete;
+};
+
+/** Structure that will parse a string and fill out some commonly used parameters per command */
+struct FCommonParameters
+{
+	int32 Limit = -1;
+
+	bool bTruncate = true;
+	bool bCSV = false;
+	
+	FName Group = NAME_None;
+	FName Class = NAME_None;
+	FName Category = NAME_None;
+
+	FString Name;
+	FString Names;
+	FString AssetName;
+	FString LogFileName;
+
+	explicit FCommonParameters(const TCHAR* Args)
+	{
+		bTruncate = !FParse::Param(Args, TEXT("notrunc"));
+		bCSV = FParse::Param(Args, TEXT("csv"));
+
+		FParse::Value(Args, TEXT("Name="), Name);
+		FParse::Value(Args, TEXT("Names="), Names);
+		FParse::Value(Args, TEXT("Limit="), Limit);
+		FParse::Value(Args, TEXT("Log="), LogFileName);
+		FParse::Value(Args, TEXT("Asset="), AssetName);
+
+		FString GroupName;
+		if (FParse::Value(Args, TEXT("Group="), GroupName))
+		{
+			Group = FName(*GroupName);
+		}
+
+		FString ClassName;
+		if (FParse::Value(Args, TEXT("Class="), ClassName))
+		{
+			Class = FName(*ClassName);
+		}
+
+		FString CategoryName;
+		if (FParse::Value(Args, TEXT("Category="), CategoryName))
+		{
+			Category = FName(*CategoryName);
+		}
+	}
+
+	FCommonParameters() = delete;
+};
+
+FAutoConsoleCommandWithWorldArgsAndOutputDevice GMemQueryUsage(
+	TEXT("MemQuery.Usage"),
+	TEXT("Name=<AssetName> Prints memory usage of the specified asset."),
+	FConsoleCommandWithWorldArgsAndOutputDeviceDelegate::CreateStatic(
+		[](const TArray<FString>& Args, UWorld* World, FOutputDevice& Ar)
+{
+	const IMemoryUsageInfoProvider* MemQueryInfoProvider = MemoryUsageQueries::GetCurrentMemoryUsageInfoProvider();
+	if (!MemQueryInfoProvider)
+	{
+		Ar.Logf(ELogVerbosity::Error, TEXT("MemoryUsageInfoProvider Error: Provider is not available. Please run with -LLM"));
+		return;
+	}
+	const FString Cmd = FString::Join(Args, TEXT(" "));
+
+	FScopedOutputDevice ScopedOutputDevice(&Ar);
+	FCommonParameters CommonArgs(*Cmd);
+
+	ScopedOutputDevice.OpenLogFile(CommonArgs.LogFileName, CommonArgs.bCSV);
+
+	FName LongPackageName;
+	if (MemoryUsageQueries::Internal::GetLongName(CommonArgs.Name, LongPackageName, &Ar))
+	{
+		uint64 ExclusiveSize = 0U;
+		uint64 InclusiveSize = 0U;
+
+		MemoryUsageQueries::GetMemoryUsage(MemQueryInfoProvider, LongPackageName, ExclusiveSize, InclusiveSize);
+		ScopedOutputDevice.GetOutputDevice().Logf(TEXT("MemoryUsage: ExclusiveSize: %.2f MiB (%.2f KiB); InclusiveSize: %.2f MiB (%.2f KiB)"),
+			ExclusiveSize / (1024.f * 1024.f),
+			ExclusiveSize / 1024.f,
+			InclusiveSize / (1024.f * 1024.f),
+			InclusiveSize / 1024.f);
+	}
+}));
+
+FAutoConsoleCommandWithWorldArgsAndOutputDevice GMemQueryCombinedUsage(
+	TEXT("MemQuery.CombinedUsage"),
+	TEXT("Names=\"<AssetName1> <AssetName2> ...\" Prints combined memory usage of the specified assets (including all dependencies)."),
+	FConsoleCommandWithWorldArgsAndOutputDeviceDelegate::CreateStatic(
+		[](const TArray<FString>& Args, UWorld* World, FOutputDevice& Ar)
+{
+	const IMemoryUsageInfoProvider* MemQueryInfoProvider = MemoryUsageQueries::GetCurrentMemoryUsageInfoProvider();
+	if (!MemQueryInfoProvider)
+	{
+		Ar.Logf(ELogVerbosity::Error, TEXT("MemoryUsageInfoProvider Error: Provider is not available. Please run with -LLM"));
+		return;
+	}
+	const FString Cmd = FString::Join(Args, TEXT(" "));
+
+	FScopedOutputDevice ScopedOutputDevice(&Ar);
+	FCommonParameters CommonArgs(*Cmd);
+
+	ScopedOutputDevice.OpenLogFile(CommonArgs.LogFileName, CommonArgs.bCSV);
+
+	if (!CommonArgs.Names.IsEmpty())
+	{
+		TArray<FString> Packages;
+		CommonArgs.Names.ParseIntoArrayWS(Packages);
+		uint64 TotalSize;
+		TSet<FName> LongPackageNames;
+		if (!MemoryUsageQueries::Internal::GetLongNames(Packages, LongPackageNames, &Ar))
+		{
+			return;
+		}
+
+		MemoryUsageQueries::GetMemoryUsageCombined(MemQueryInfoProvider, LongPackageNames.Array(), TotalSize);
+		ScopedOutputDevice.GetOutputDevice().Logf(TEXT("MemoryUsageCombined: TotalSize: %.2f MiB (%.2f KiB)"), TotalSize / (1024.f * 1024.f), TotalSize / 1024.f);
+	}
+}));
+
+FAutoConsoleCommandWithWorldArgsAndOutputDevice GMemQuerySharedUsage(
+	TEXT("MemQuery.SharedUsage"),
+	TEXT("Names=\"<AssetName1> <AssetName2> ...\" Prints shared memory usage of the specified assets (including only dependencies shared by the specified assets)."),
+	FConsoleCommandWithWorldArgsAndOutputDeviceDelegate::CreateStatic(
+		[](const TArray<FString>& Args, UWorld* World, FOutputDevice& Ar)
+{
+	const IMemoryUsageInfoProvider* MemQueryInfoProvider = MemoryUsageQueries::GetCurrentMemoryUsageInfoProvider();
+	if (!MemQueryInfoProvider)
+	{
+		Ar.Logf(ELogVerbosity::Error, TEXT("MemoryUsageInfoProvider Error: Provider is not available. Please run with -LLM"));
+		return;
+	}
+	const FString Cmd = FString::Join(Args, TEXT(" "));
+
+	FScopedOutputDevice ScopedOutputDevice(&Ar);
+	FCommonParameters CommonArgs(*Cmd);
+
+	ScopedOutputDevice.OpenLogFile(CommonArgs.LogFileName, CommonArgs.bCSV);
+
+	if (!CommonArgs.Names.IsEmpty())
+	{
+		TArray<FString> Packages;
+		CommonArgs.Names.ParseIntoArrayWS(Packages);
+		uint64 SharedSize;
+		TSet<FName> LongPackageNames;
+		if (!MemoryUsageQueries::Internal::GetLongNames(Packages, LongPackageNames, &Ar))
+		{
+			return;
+		}
+
+		MemoryUsageQueries::GetMemoryUsageShared(MemQueryInfoProvider, LongPackageNames.Array(), SharedSize);
+		ScopedOutputDevice.GetOutputDevice().Logf(TEXT("MemoryUsageShared: SharedSize: %.2f MiB (%.2f KiB)"), SharedSize / (1024.f * 1024.f), SharedSize / 1024.f);
+	}
+}));
+
+FAutoConsoleCommandWithWorldArgsAndOutputDevice GMemQueryUniqueUsage(
+	TEXT("MemQuery.UniqueUsage"),
+	TEXT("Names=\"<AssetName1> <AssetName2> ...\" Prints unique memory usage of the specified assets (including only dependencies unique to the specified assets)."),
+	FConsoleCommandWithWorldArgsAndOutputDeviceDelegate::CreateStatic(
+		[](const TArray<FString>& Args, UWorld* World, FOutputDevice& Ar)
+{
+	const IMemoryUsageInfoProvider* MemQueryInfoProvider = MemoryUsageQueries::GetCurrentMemoryUsageInfoProvider();
+	if (!MemQueryInfoProvider)
+	{
+		Ar.Logf(ELogVerbosity::Error, TEXT("MemoryUsageInfoProvider Error: Provider is not available. Please run with -LLM"));
+		return;
+	}
+	const FString Cmd = FString::Join(Args, TEXT(" "));
+
+	FScopedOutputDevice ScopedOutputDevice(&Ar);
+	FCommonParameters CommonArgs(*Cmd);
+
+	ScopedOutputDevice.OpenLogFile(CommonArgs.LogFileName, CommonArgs.bCSV);
+
+	if (!CommonArgs.Names.IsEmpty())
+	{
+		TArray<FString> Packages;
+		CommonArgs.Names.ParseIntoArrayWS(Packages);
+		uint64 UniqueSize = 0U;
+		TSet<FName> LongPackageNames;
+		if (!MemoryUsageQueries::Internal::GetLongNames(Packages, LongPackageNames, &Ar))
+		{
+			return;
+		}
+
+		MemoryUsageQueries::GetMemoryUsageUnique(MemQueryInfoProvider, LongPackageNames.Array(), UniqueSize);
+		ScopedOutputDevice.GetOutputDevice().Logf(TEXT("MemoryUsageUnique: UniqueSize: %.2f MiB (%.2f KiB)"), UniqueSize / (1024.f * 1024.f), UniqueSize / 1024.f);
+	}
+}));
+
+FAutoConsoleCommandWithWorldArgsAndOutputDevice GMemQueryCommonUsage(
+	TEXT("MemQuery.CommonUsage"),
+	TEXT("Names=\"<AssetName1> <AssetName2> ...\" Prints common memory usage of the specified assets (including only dependencies that are not unique to the specified assets)."),
+	FConsoleCommandWithWorldArgsAndOutputDeviceDelegate::CreateStatic(
+		[](const TArray<FString>& Args, UWorld* World, FOutputDevice& Ar)
+{
+	const IMemoryUsageInfoProvider* MemQueryInfoProvider = MemoryUsageQueries::GetCurrentMemoryUsageInfoProvider();
+	if (!MemQueryInfoProvider)
+	{
+		Ar.Logf(ELogVerbosity::Error, TEXT("MemoryUsageInfoProvider Error: Provider is not available. Please run with -LLM"));
+		return;
+	}
+	const FString Cmd = FString::Join(Args, TEXT(" "));
+
+	FScopedOutputDevice ScopedOutputDevice(&Ar);
+	FCommonParameters CommonArgs(*Cmd);
+
+	ScopedOutputDevice.OpenLogFile(CommonArgs.LogFileName, CommonArgs.bCSV);
+
+	if (!CommonArgs.Names.IsEmpty())
+	{
+		TArray<FString> Packages;
+		CommonArgs.Names.ParseIntoArrayWS(Packages);
+		uint64 CommonSize = 0U;
+		TSet<FName> LongPackageNames;
+		if (!MemoryUsageQueries::Internal::GetLongNames(Packages, LongPackageNames, &Ar))
+		{
+			return;
+		}
+
+		MemoryUsageQueries::GetMemoryUsageCommon(MemQueryInfoProvider, LongPackageNames.Array(), CommonSize);
+		ScopedOutputDevice.GetOutputDevice().Logf(TEXT("MemoryUsageCommon: CommonSize: %.2f MiB (%.2f KiB)"), CommonSize / (1024.f * 1024.f), CommonSize / 1024.f);
+	}
+}));
+
+FAutoConsoleCommandWithWorldArgsAndOutputDevice GMemQueryDependencies(
+	TEXT("MemQuery.Dependencies"),
+	TEXT("Name=<AssetName> Limit=<n> Lists dependencies of the specified asset, sorted by size."),
+	FConsoleCommandWithWorldArgsAndOutputDeviceDelegate::CreateStatic(
+		[](const TArray<FString>& Args, UWorld* World, FOutputDevice& Ar)
+{
+	const IMemoryUsageInfoProvider* MemQueryInfoProvider = MemoryUsageQueries::GetCurrentMemoryUsageInfoProvider();
+	if (!MemQueryInfoProvider)
+	{
+		Ar.Logf(ELogVerbosity::Error, TEXT("MemoryUsageInfoProvider Error: Provider is not available. Please run with -LLM"));
+		return;
+	}
+	const FString Cmd = FString::Join(Args, TEXT(" "));
+
+	FScopedOutputDevice ScopedOutputDevice(&Ar);
+	FCommonParameters CommonArgs(*Cmd);
+
+	ScopedOutputDevice.OpenLogFile(CommonArgs.LogFileName, CommonArgs.bCSV);
+
+	FName LongPackageName;
+	if (MemoryUsageQueries::Internal::GetLongName(CommonArgs.Name, LongPackageName, &Ar))
+	{
+		TMap<FName, uint64> DependenciesWithSize;
+
+		MemoryUsageQueries::GetDependenciesWithSize(MemQueryInfoProvider, LongPackageName, DependenciesWithSize);
+		MemoryUsageQueries::Internal::PrintTagsWithSize(ScopedOutputDevice.GetOutputDevice(), DependenciesWithSize, TEXT("Dependencies"), CommonArgs.bTruncate, CommonArgs.Limit, CommonArgs.bCSV);
+	}
+}));
+
+FAutoConsoleCommandWithWorldArgsAndOutputDevice GMemQueryCombinedDependencies(
+	TEXT("MemQuery.CombinedDependencies"),
+	TEXT("Names=\"<AssetName1> <AssetName2> ...\" Limit=<n> Lists n largest dependencies of the specified assets, sorted by size."),
+	FConsoleCommandWithWorldArgsAndOutputDeviceDelegate::CreateStatic(
+		[](const TArray<FString>& Args, UWorld* World, FOutputDevice& Ar)
+{
+	const IMemoryUsageInfoProvider* MemQueryInfoProvider = MemoryUsageQueries::GetCurrentMemoryUsageInfoProvider();
+	if (!MemQueryInfoProvider)
+	{
+		Ar.Logf(ELogVerbosity::Error, TEXT("MemoryUsageInfoProvider Error: Provider is not available. Please run with -LLM"));
+		return;
+	}
+	const FString Cmd = FString::Join(Args, TEXT(" "));
+
+	FScopedOutputDevice ScopedOutputDevice(&Ar);
+	FCommonParameters CommonArgs(*Cmd);
+
+	ScopedOutputDevice.OpenLogFile(CommonArgs.LogFileName, CommonArgs.bCSV);
+
+	if (!CommonArgs.Names.IsEmpty())
+	{
+		TArray<FString> Packages;
+		CommonArgs.Names.ParseIntoArrayWS(Packages);
+		TMap<FName, uint64> CombinedDependenciesWithSize;
+		TSet<FName> LongPackageNames;
+		if (!MemoryUsageQueries::Internal::GetLongNames(Packages, LongPackageNames, &Ar))
+		{
+			return;
+		}
+
+		MemoryUsageQueries::GetDependenciesWithSizeCombined(MemQueryInfoProvider, LongPackageNames.Array(), CombinedDependenciesWithSize);
+		MemoryUsageQueries::Internal::PrintTagsWithSize(ScopedOutputDevice.GetOutputDevice(), CombinedDependenciesWithSize, TEXT("Combined Dependencies"), CommonArgs.bTruncate, CommonArgs.Limit, CommonArgs.bCSV);
+	}
+}));
+
+FAutoConsoleCommandWithWorldArgsAndOutputDevice GMemQuerySharedDependencies(
+	TEXT("MemQuery.SharedDependencies"),
+	TEXT("Names=\"<AssetName1> <AssetName2> ...\" Limit=<n> Lists n largest dependencies that are shared by the specified assets, sorted by size."),
+	FConsoleCommandWithWorldArgsAndOutputDeviceDelegate::CreateStatic(
+		[](const TArray<FString>& Args, UWorld* World, FOutputDevice& Ar)
+{
+	const IMemoryUsageInfoProvider* MemQueryInfoProvider = MemoryUsageQueries::GetCurrentMemoryUsageInfoProvider();
+	if (!MemQueryInfoProvider)
+	{
+		Ar.Logf(ELogVerbosity::Error, TEXT("MemoryUsageInfoProvider Error: Provider is not available. Please run with -LLM"));
+		return;
+	}
+	const FString Cmd = FString::Join(Args, TEXT(" "));
+
+	FScopedOutputDevice ScopedOutputDevice(&Ar);
+	FCommonParameters CommonArgs(*Cmd);
+
+	ScopedOutputDevice.OpenLogFile(CommonArgs.LogFileName, CommonArgs.bCSV);
+
+	if (!CommonArgs.Names.IsEmpty())
+	{
+		TArray<FString> Packages;
+		CommonArgs.Names.ParseIntoArrayWS(Packages);
+		TMap<FName, uint64> SharedDependenciesWithSize;
+		TSet<FName> LongPackageNames;
+		if (!MemoryUsageQueries::Internal::GetLongNames(Packages, LongPackageNames, &Ar))
+		{
+			return;
+		}
+
+		MemoryUsageQueries::GetDependenciesWithSizeShared(MemQueryInfoProvider, LongPackageNames.Array(), SharedDependenciesWithSize);
+		MemoryUsageQueries::Internal::PrintTagsWithSize(ScopedOutputDevice.GetOutputDevice(), SharedDependenciesWithSize, TEXT("Shared Dependencies"), CommonArgs.bTruncate, CommonArgs.Limit, CommonArgs.bCSV);
+	}
+}));
+
+FAutoConsoleCommandWithWorldArgsAndOutputDevice GMemQueryUniqueDependencies(
+	TEXT("MemQuery.UniqueDependencies"),
+	TEXT("Names=\"<AssetName1> <AssetName2> ...\" Limit=<n> Lists n largest dependencies that are unique to the specified assets, sorted by size."),
+	FConsoleCommandWithWorldArgsAndOutputDeviceDelegate::CreateStatic(
+		[](const TArray<FString>& Args, UWorld* World, FOutputDevice& Ar)
+{
+	const IMemoryUsageInfoProvider* MemQueryInfoProvider = MemoryUsageQueries::GetCurrentMemoryUsageInfoProvider();
+	if (!MemQueryInfoProvider)
+	{
+		Ar.Logf(ELogVerbosity::Error, TEXT("MemoryUsageInfoProvider Error: Provider is not available. Please run with -LLM"));
+		return;
+	}
+	const FString Cmd = FString::Join(Args, TEXT(" "));
+
+	FScopedOutputDevice ScopedOutputDevice(&Ar);
+	FCommonParameters CommonArgs(*Cmd);
+
+	ScopedOutputDevice.OpenLogFile(CommonArgs.LogFileName, CommonArgs.bCSV);
+
+	if (!CommonArgs.Names.IsEmpty())
+	{
+		TArray<FString> Packages;
+		CommonArgs.Names.ParseIntoArrayWS(Packages);
+		TMap<FName, uint64> UniqueDependenciesWithSize;
+		TSet<FName> LongPackageNames;
+		if (!MemoryUsageQueries::Internal::GetLongNames(Packages, LongPackageNames, &Ar))
+		{
+			return;
+		}
+
+		MemoryUsageQueries::GetDependenciesWithSizeUnique(MemQueryInfoProvider, LongPackageNames.Array(), UniqueDependenciesWithSize);
+		MemoryUsageQueries::Internal::PrintTagsWithSize(ScopedOutputDevice.GetOutputDevice(), UniqueDependenciesWithSize, TEXT("Unique Dependencies"), CommonArgs.bTruncate, CommonArgs.Limit, CommonArgs.bCSV);
+	}
+}));
+
+FAutoConsoleCommandWithWorldArgsAndOutputDevice GMemQueryCommonDependencies(
+	TEXT("MemQuery.CommonDependencies"),
+	TEXT("Names=\"<AssetName1> <AssetName2> ...\" Limit=<n> Lists n largest dependencies that are NOT unique to the specified assets, sorted by size."),
+	FConsoleCommandWithWorldArgsAndOutputDeviceDelegate::CreateStatic(
+		[](const TArray<FString>& Args, UWorld* World, FOutputDevice& Ar)
+{
+	const IMemoryUsageInfoProvider* MemQueryInfoProvider = MemoryUsageQueries::GetCurrentMemoryUsageInfoProvider();
+	if (!MemQueryInfoProvider)
+	{
+		Ar.Logf(ELogVerbosity::Error, TEXT("MemoryUsageInfoProvider Error: Provider is not available. Please run with -LLM"));
+		return;
+	}
+	const FString Cmd = FString::Join(Args, TEXT(" "));
+
+	FScopedOutputDevice ScopedOutputDevice(&Ar);
+	FCommonParameters CommonArgs(*Cmd);
+
+	ScopedOutputDevice.OpenLogFile(CommonArgs.LogFileName, CommonArgs.bCSV);
+
+	if (!CommonArgs.Names.IsEmpty())
+	{
+		TArray<FString> Packages;
+		CommonArgs.Names.ParseIntoArrayWS(Packages);
+		TMap<FName, uint64> CommonDependenciesWithSize;
+		TSet<FName> LongPackageNames;
+		if (!MemoryUsageQueries::Internal::GetLongNames(Packages, LongPackageNames, &Ar))
+		{
+			return;
+		}
+
+		MemoryUsageQueries::GetDependenciesWithSizeCommon(MemQueryInfoProvider, LongPackageNames.Array(), CommonDependenciesWithSize);
+		MemoryUsageQueries::Internal::PrintTagsWithSize(ScopedOutputDevice.GetOutputDevice(), CommonDependenciesWithSize, TEXT("Common Dependencies"), CommonArgs.bTruncate, CommonArgs.Limit, CommonArgs.bCSV);
+	}
+}));
+
+FAutoConsoleCommandWithWorldArgsAndOutputDevice GMemQuerySavings(
+	TEXT("MemQuery.Savings"),
+	TEXT("Limit=<n> Lists potential savings among %s. How much memory can be saved it we delete certain object."),
+	FConsoleCommandWithWorldArgsAndOutputDeviceDelegate::CreateStatic(
+		[](const TArray<FString>& Args, UWorld* World, FOutputDevice& Ar)
+{
+	const IMemoryUsageInfoProvider* MemQueryInfoProvider = MemoryUsageQueries::GetCurrentMemoryUsageInfoProvider();
+	if (!MemQueryInfoProvider)
+	{
+		Ar.Logf(ELogVerbosity::Error, TEXT("MemoryUsageInfoProvider Error: Provider is not available. Please run with -LLM"));
+		return;
+	}
+	const FString Cmd = FString::Join(Args, TEXT(" "));
+
+	FScopedOutputDevice ScopedOutputDevice(&Ar);
+	FCommonParameters CommonArgs(*Cmd);
+
+	ScopedOutputDevice.OpenLogFile(CommonArgs.LogFileName, CommonArgs.bCSV);
+
+	const UMemoryUsageQueriesConfig* Config = GetDefault<UMemoryUsageQueriesConfig>();
+
+	// Necessary for FParse::Command
+	const TCHAR* Command = *Cmd;
+
+	for (auto It = Config->SavingsPresets.CreateConstIterator(); It; ++It)
+	{
+		if (!FParse::Command(&Command, *It.Key()))
+		{
+			continue;
+		}
+
+		TMap<FName, uint64> PresetSavings;
+		TSet<FName> Packages;
+
+		UClass* SavingsClass = FindObject<UClass>(nullptr, *It.Value());
+		if (SavingsClass != nullptr)
+		{
+			TArray<UClass*> DerivedClasses;
+			TArray<UClass*> DerivedResults;
+
+			GetDerivedClasses(SavingsClass, DerivedClasses, true);
+
+			for (UClass* DerivedClass : DerivedClasses)
+			{
+				UBlueprintGeneratedClass* BPClass = Cast<UBlueprintGeneratedClass>(DerivedClass);
+				if (BPClass != nullptr)
+				{
+					DerivedResults.Reset();
+					GetDerivedClasses(BPClass, DerivedResults, false);
+					if (DerivedResults.Num() == 0)
+					{
+						Packages.Add(DerivedClass->GetPackage()->GetFName());
+					}
+				}
+			}
+		}
+
+		for (const auto& Package : Packages)
+		{
+			uint64 Size = 0;
+			MemoryUsageQueries::GetMemoryUsageUnique(MemQueryInfoProvider, TArray{Package}, Size);
+			PresetSavings.Add(Package, Size);
+		}
+
+		PresetSavings.ValueSort(TGreater<uint64>());
+		MemoryUsageQueries::Internal::PrintTagsWithSize(ScopedOutputDevice.GetOutputDevice(), PresetSavings, *FString::Printf(TEXT("possible savings")), CommonArgs.bTruncate, CommonArgs.bCSV);
+	}
+}));
+
+#if ENABLE_LOW_LEVEL_MEM_TRACKER
+
+FAutoConsoleCommandWithWorldArgsAndOutputDevice GMemQueryListAssets(
+	TEXT("MemQuery.ListAssets"),
+	TEXT("Asset=<AssetNameSubstring> Group=<GroupName> Class=<ClassName> Limit=<n> Lists n largest assets."),
+	FConsoleCommandWithWorldArgsAndOutputDeviceDelegate::CreateStatic(
+		[](const TArray<FString>& Args, UWorld* World, FOutputDevice& Ar)
+{
+	const IMemoryUsageInfoProvider* MemQueryInfoProvider = MemoryUsageQueries::GetCurrentMemoryUsageInfoProvider();
+	if (!MemQueryInfoProvider)
+	{
+		Ar.Logf(ELogVerbosity::Error, TEXT("MemoryUsageInfoProvider Error: Provider is not available. Please run with -LLM"));
+		return;
+	}
+	const FString Cmd = FString::Join(Args, TEXT(" "));
+
+	FScopedOutputDevice ScopedOutputDevice(&Ar);
+	FCommonParameters CommonArgs(*Cmd);
+
+	ScopedOutputDevice.OpenLogFile(CommonArgs.LogFileName, CommonArgs.bCSV);
+
+	bool bSuccess;
+	TMap<FName, uint64> AssetsWithSize;
+
+	if (CommonArgs.Group != NAME_None || CommonArgs.Class != NAME_None)
+	{
+		bSuccess = MemoryUsageQueries::GetFilteredPackagesWithSize(AssetsWithSize, CommonArgs.Group, CommonArgs.AssetName, CommonArgs.Class);
+	}
+	else
+	{
+		// TODO - Implement using faster path if there are no group / class filters
+		bSuccess = MemoryUsageQueries::GetFilteredPackagesWithSize(AssetsWithSize, CommonArgs.Group, CommonArgs.AssetName, CommonArgs.Class);
+	}
+
+	if (bSuccess)
+	{
+		MemoryUsageQueries::Internal::PrintTagsWithSize(ScopedOutputDevice.GetOutputDevice(), AssetsWithSize, TEXT("largest assets"), CommonArgs.bTruncate, CommonArgs.Limit, CommonArgs.bCSV);
+	}
+}));
+
+FAutoConsoleCommandWithWorldArgsAndOutputDevice GMemQueryListAssetsCategorized(
+	TEXT("MemQuery.ListAssetsCategorized"),
+	TEXT("Asset=<AssetNameSubstring> Group=<GroupName> Class=<ClassName> Category=<CategoryName(None,Assets,AssetClasses)> Limit=<n> Lists n largest assets categorized by Category."),
+	FConsoleCommandWithWorldArgsAndOutputDeviceDelegate::CreateStatic(
+		[](const TArray<FString>& Args, UWorld* World, FOutputDevice& Ar)
+{
+	const IMemoryUsageInfoProvider* MemQueryInfoProvider = MemoryUsageQueries::GetCurrentMemoryUsageInfoProvider();
+	if (!MemQueryInfoProvider)
+	{
+		Ar.Logf(ELogVerbosity::Error, TEXT("MemoryUsageInfoProvider Error: Provider is not available. Please run with -LLM"));
+		return;
+	}
+	const FString Cmd = FString::Join(Args, TEXT(" "));
+
+	FScopedOutputDevice ScopedOutputDevice(&Ar);
+	FCommonParameters CommonArgs(*Cmd);
+
+	ScopedOutputDevice.OpenLogFile(CommonArgs.LogFileName, CommonArgs.bCSV);
+
+	bool bSuccess;
+	TMap<FName, uint64> AssetsWithSize;
+	bSuccess = MemoryUsageQueries::GetFilteredPackagesCategorizedWithSize(AssetsWithSize, CommonArgs.Group, CommonArgs.AssetName, CommonArgs.Class, CommonArgs.Category, &ScopedOutputDevice.GetOutputDevice());
+
+	if (bSuccess)
+	{
+		MemoryUsageQueries::Internal::PrintTagsWithSize(ScopedOutputDevice.GetOutputDevice(), AssetsWithSize, TEXT("largest assets"), CommonArgs.bTruncate, CommonArgs.Limit, CommonArgs.bCSV);
+	}
+}));
+
+FAutoConsoleCommandWithWorldArgsAndOutputDevice GMemQueryListClasses(
+	TEXT("MemQuery.ListClasses"),
+	TEXT("Group=<GroupName> Asset=<AssetName> Limit=<n> Lists n largest classes."),
+	FConsoleCommandWithWorldArgsAndOutputDeviceDelegate::CreateStatic(
+		[](const TArray<FString>& Args, UWorld* World, FOutputDevice& Ar)
+{
+	const IMemoryUsageInfoProvider* MemQueryInfoProvider = MemoryUsageQueries::GetCurrentMemoryUsageInfoProvider();
+	if (!MemQueryInfoProvider)
+	{
+		Ar.Logf(ELogVerbosity::Error, TEXT("MemoryUsageInfoProvider Error: Provider is not available. Please run with -LLM"));
+		return;
+	}
+	const FString Cmd = FString::Join(Args, TEXT(" "));
+
+	FScopedOutputDevice ScopedOutputDevice(&Ar);
+	FCommonParameters CommonArgs(*Cmd);
+
+	ScopedOutputDevice.OpenLogFile(CommonArgs.LogFileName, CommonArgs.bCSV);
+
+	TMap<FName, uint64> ClassesWithSize;
+	if (MemoryUsageQueries::GetFilteredClassesWithSize(ClassesWithSize, CommonArgs.Group, CommonArgs.AssetName, &Ar))
+	{
+		MemoryUsageQueries::Internal::PrintTagsWithSize(ScopedOutputDevice.GetOutputDevice(), ClassesWithSize, *FString::Printf(TEXT("Largest Classes")), CommonArgs.bTruncate, CommonArgs.Limit, CommonArgs.bCSV);
+	}
+}));
+
+FAutoConsoleCommandWithWorldArgsAndOutputDevice GMemQueryListGroups(
+	TEXT("MemQuery.ListGroups"),
+	TEXT("Asset=<AssetName> Class=<ClassName> Limit=<n> Lists n largest groups."),
+	FConsoleCommandWithWorldArgsAndOutputDeviceDelegate::CreateStatic(
+		[](const TArray<FString>& Args, UWorld* World, FOutputDevice& Ar)
+{
+	const IMemoryUsageInfoProvider* MemQueryInfoProvider = MemoryUsageQueries::GetCurrentMemoryUsageInfoProvider();
+	if (!MemQueryInfoProvider)
+	{
+		Ar.Logf(ELogVerbosity::Error, TEXT("MemoryUsageInfoProvider Error: Provider is not available. Please run with -LLM"));
+		return;
+	}
+	const FString Cmd = FString::Join(Args, TEXT(" "));
+
+	FScopedOutputDevice ScopedOutputDevice(&Ar);
+	FCommonParameters CommonArgs(*Cmd);
+
+	ScopedOutputDevice.OpenLogFile(CommonArgs.LogFileName, CommonArgs.bCSV);
+
+	TMap<FName, uint64> GroupsWithSize;
+	if (MemoryUsageQueries::GetFilteredGroupsWithSize(GroupsWithSize, CommonArgs.AssetName, CommonArgs.Class, &Ar))
+	{
+		MemoryUsageQueries::Internal::PrintTagsWithSize(ScopedOutputDevice.GetOutputDevice(), GroupsWithSize, *FString::Printf(TEXT("Largest Groups")), CommonArgs.bTruncate, CommonArgs.Limit, CommonArgs.bCSV);
+	}
+}));
+
+// Will return true if the Package Name matches any of the conditions in the array of Paths
+static bool PackageNameMatches(const FString& PackageName, const TArray<FString>& Conditions)
+{
+	for (const FString& Condition : Conditions)
+	{
+		if ((FWildcardString::ContainsWildcards(*Condition) && FWildcardString::IsMatch(*Condition, *PackageName)) || (PackageName.Contains(Condition)))
+		{
+			return true;
+		}
+	}
+
+	return false;
+}
+
+FAutoConsoleCommandWithWorldArgsAndOutputDevice GMemQueryCollections(
+	TEXT("MemQuery.Collection"),
+	TEXT("Lists memory used by a collection. Can show dependency breakdown. Pass -showdeps to list dependencies."),
+	FConsoleCommandWithWorldArgsAndOutputDeviceDelegate::CreateStatic(
+		[](const TArray<FString>& Args, UWorld* World, FOutputDevice& Ar)
+{
+	const IMemoryUsageInfoProvider* MemQueryInfoProvider = MemoryUsageQueries::GetCurrentMemoryUsageInfoProvider();
+	if (!MemQueryInfoProvider)
+	{
+		Ar.Logf(ELogVerbosity::Error, TEXT("MemoryUsageInfoProvider Error: Provider is not available. Please run with -LLM"));
+		return;
+	}
+	const FString Cmd = FString::Join(Args, TEXT(" "));
+
+	FScopedOutputDevice ScopedOutputDevice(&Ar);
+	FCommonParameters CommonArgs(*Cmd);
+
+	ScopedOutputDevice.OpenLogFile(CommonArgs.LogFileName, CommonArgs.bCSV);
+
+	const bool bShowDependencies = FParse::Param(*Cmd, TEXT("ShowDeps"));
+
+	// Necessary for FParse::Command
+	const TCHAR* Command = *Cmd;
+
+	FOutputDevice* CurrentOutputDevice = &ScopedOutputDevice.GetOutputDevice();
+
+	const UMemoryUsageQueriesConfig* Config = GetDefault<UMemoryUsageQueriesConfig>();
+	for (auto It = Config->Collections.CreateConstIterator(); It; ++It)
+	{
+		const FCollectionInfo& CollectionInfo = *It;
+
+		if (!FParse::Command(&Command, *CollectionInfo.Name))
+		{
+			continue;
+		}
+
+		// Retrieve a list of all assets that have allocations we are currently tracking.
+		TMap<FName, uint64> AssetsWithSize;
+		bool bSuccess = MemoryUsageQueries::GetFilteredPackagesWithSize(AssetsWithSize, NAME_None, "", NAME_None);
+
+		if (!bSuccess)
+		{
+			Ar.Logf(TEXT("Failed to gather assets for Collection %s"), *CollectionInfo.Name);
+			break;
+		}
+
+		// See if any of the asset paths match those of our matching paths and are valid
+		TArray<FName> PackageNames;
+
+		TMap<FName, FAssetMemoryDetails> AssetMemoryMap;
+		for (const TPair<FName, uint64>& AssetSizePair : AssetsWithSize)
+		{
+			const FString& PackageName = AssetSizePair.Key.ToString();
+
+			if (!FPackageName::IsValidLongPackageName(PackageName))
+			{
+				continue;
+			}
+
+			// If path is Included and NOT Excluded, it's a valid asset to consider
+			if (PackageNameMatches(PackageName, CollectionInfo.Includes) && !PackageNameMatches(PackageName, CollectionInfo.Excludes))
+			{
+				FAssetMemoryDetails& AssetMemory = AssetMemoryMap.Add(AssetSizePair.Key);
+				AssetMemory.MemoryBreakdown.ExclusiveSize = AssetSizePair.Value;
+
+				FName LongPackageName;
+				if (!MemoryUsageQueries::Internal::GetLongName(PackageName, LongPackageName, &Ar))
+				{
+					Ar.Logf(TEXT("Failed to get long package name for Asset %s"), *PackageName);
+				}
+				MemoryUsageQueries::Internal::GetTransitiveDependencies(LongPackageName, AssetMemory.Dependencies);
+				AssetMemory.Dependencies.Add(LongPackageName);
+
+				PackageNames.Add(LongPackageName);
+			}
+		}
+
+		// Gather list of dependencies. Internal dependencies are confined only to the set of packages passed in.
+		// External are dependencies that have additional references outside the set of packages passed in.
+		TMap<FName, uint64> InternalDependencies;
+		TMap<FName, uint64> ExternalDependencies;
+		MemoryUsageQueries::GatherDependenciesForPackages(MemQueryInfoProvider, PackageNames, InternalDependencies, ExternalDependencies);
+
+		uint64 TotalCollectionSize = 0;
+
+		// Determine which category where each assets dependency should reside
+		for (TPair<FName, FAssetMemoryDetails>& Asset : AssetMemoryMap)
+		{
+			FAssetMemoryDetails& AssetMemory = Asset.Value;
+			FAssetMemoryBreakdown& AssetMemoryDetails = AssetMemory.MemoryBreakdown;
+
+			for (FName& Dep : Asset.Value.Dependencies)
+			{
+				// Don't want to count asset itself, plus some dependencies might refer to other assets in the map
+				if (AssetMemoryMap.Contains(Dep))
+				{
+					continue;
+				}
+
+				FAssetMemoryBreakdown DependencyMemory;
+				uint64* UniqueMemory = InternalDependencies.Find(Dep);
+				uint64* SharedMemory = ExternalDependencies.Find(Dep);
+				bool bRecordDependency = false;
+
+				if (UniqueMemory != nullptr && *UniqueMemory != 0)
+				{
+					DependencyMemory.UniqueSize = *UniqueMemory;
+					AssetMemoryDetails.UniqueSize += DependencyMemory.UniqueSize;
+					AssetMemory.UniqueRefCount++;
+					bRecordDependency = true;
+				}
+
+				if (SharedMemory != nullptr && *SharedMemory != 0)
+				{
+					DependencyMemory.SharedSize = *SharedMemory;
+					AssetMemoryDetails.SharedSize += DependencyMemory.SharedSize;
+					AssetMemory.SharedRefCount++;
+					bRecordDependency = true;
+				}
+
+				if (bRecordDependency)
+				{
+					AssetMemory.DependenciesToMemoryMap.Add(Dep, DependencyMemory);
+				}
+			}
+
+			AssetMemoryDetails.TotalSize = AssetMemoryDetails.ExclusiveSize + AssetMemoryDetails.UniqueSize;
+			TotalCollectionSize += AssetMemoryDetails.TotalSize;
+		}
+
+		// Sort by TotalSize
+		AssetMemoryMap.ValueSort([](const FAssetMemoryDetails& A, const FAssetMemoryDetails& B) {
+			return A.MemoryBreakdown.TotalSize > B.MemoryBreakdown.TotalSize;
+			});
+
+		if (CommonArgs.bCSV)
+		{
+			CurrentOutputDevice->Logf(TEXT(",Asset,Exclusive KiB,Unique Refs KiB,Unique Ref Count,Shared Refs KiB,Shared Ref Count,Total KiB"));
+		}
+		else
+		{
+			CurrentOutputDevice->Logf(
+				TEXT(" %100s %20s %20s %15s %20s %15s %25s"),
+				TEXT("Asset"),
+				TEXT("Exclusive KiB"),
+				TEXT("Unique Refs KiB"),
+				TEXT("Unique Ref Count"),
+				TEXT("Shared Refs KiB"),
+				TEXT("Shared Ref Count"),
+				TEXT("Total KiB")
+			);
+		}
+
+		// Asset listing
+		for (TPair<FName, FAssetMemoryDetails>& Asset : AssetMemoryMap)
+		{
+			FAssetMemoryDetails& AssetMemory = Asset.Value;
+			FAssetMemoryBreakdown& AssetMemoryDetails = AssetMemory.MemoryBreakdown;
+
+			if (CommonArgs.bCSV)
+			{
+				CurrentOutputDevice->Logf(TEXT(",%s,%.2f,%.2f,%d,%.2f,%d,%.2f"), *Asset.Key.ToString(),
+					AssetMemoryDetails.ExclusiveSize / 1024.f,
+					AssetMemoryDetails.UniqueSize / 1024.f,
+					AssetMemory.UniqueRefCount,
+					AssetMemoryDetails.SharedSize / 1024.f,
+					AssetMemory.SharedRefCount,
+					AssetMemoryDetails.TotalSize / 1024.f);
+			}
+			else
+			{
+				CurrentOutputDevice->Logf(
+					TEXT(" %100s %20.2f %20.2f %15d %20.2f %15d %25.2f"),
+					*Asset.Key.ToString(),
+					AssetMemoryDetails.ExclusiveSize / 1024.f,
+					AssetMemoryDetails.SharedSize / 1024.f,
+					AssetMemory.SharedRefCount,
+					AssetMemoryDetails.UniqueSize / 1024.f,
+					AssetMemory.UniqueRefCount,
+					AssetMemoryDetails.TotalSize / 1024.f
+				);
+			}
+		}
+
+		// Asset dependencies listing
+		if (bShowDependencies)
+		{
+			if (CommonArgs.bCSV)
+			{
+				CurrentOutputDevice->Logf(TEXT(",Asset,Dependency,Unique KiB,Shared KiB"));
+			}
+			else
+			{
+				CurrentOutputDevice->Logf(TEXT(" %100s %100s %20s %20s"),
+					TEXT("Asset"),
+					TEXT("Dependency"),
+					TEXT("Unique KiB"),
+					TEXT("Shared KiB")
+				);
+			}
+
+			for (TPair<FName, FAssetMemoryDetails>& Asset : AssetMemoryMap)
+			{
+				for (TPair<FName, FAssetMemoryBreakdown>& Dep : Asset.Value.DependenciesToMemoryMap)
+				{
+					const FString DependencyAssetName = Dep.Key.ToString();
+					const FAssetMemoryBreakdown& DepedencyMemoryDetails = Dep.Value;
+
+					if (CommonArgs.bCSV)
+					{
+						CurrentOutputDevice->Logf(TEXT(",%s,%s,%.2f,%.2f"), *Asset.Key.ToString(), *DependencyAssetName,
+							DepedencyMemoryDetails.UniqueSize / 1024.f, DepedencyMemoryDetails.SharedSize / 1024.f);
+					}
+					else
+					{
+						CurrentOutputDevice->Logf(TEXT(" %100s %100s %20.2f %20.2f"),
+							*Asset.Key.ToString(),
+							*DependencyAssetName,
+							DepedencyMemoryDetails.UniqueSize / 1024.f,
+							DepedencyMemoryDetails.SharedSize / 1024.f
+						);
+					}
+				}
+			}
+		}
+
+		if (CommonArgs.bCSV)
+		{
+			CurrentOutputDevice->Logf(TEXT(",TOTAL KiB,%.2f"), TotalCollectionSize / 1024.f);
+		}
+		else
+		{
+			CurrentOutputDevice->Logf(TEXT("TOTAL KiB: %.2f"), TotalCollectionSize / 1024.f);
+		}
+	}
+}));
+
+#endif
diff --git a/Engine/Plugins/MemoryUsageQueries/Source/MemoryUsageQueries/Private/MemoryUsageQueriesPrivate.h b/Engine/Plugins/MemoryUsageQueries/Source/MemoryUsageQueries/Private/MemoryUsageQueriesPrivate.h
--- a/Engine/Plugins/MemoryUsageQueries/Source/MemoryUsageQueries/Private/MemoryUsageQueriesPrivate.h
+++ b/Engine/Plugins/MemoryUsageQueries/Source/MemoryUsageQueries/Private/MemoryUsageQueriesPrivate.h
@@ -10,59 +10,65 @@ namespace MemoryUsageQueries::Internal
 
 class FMemoryUsageReferenceProcessor : public FSimpleReferenceProcessorBase
 {
-public:
-	enum EMode
+private:
+	struct FAssetSet
 	{
-		Full,
-		Excluding
+		TBitArray<> Excluded;
+		TBitArray<> ReachableExcluded;
 	};
-
-private:
-	TBitArray<> Excluded;
+	TArray<FAssetSet> AssetSets;
 
 	TBitArray<> ReachableFull;
-	TBitArray<> ReachableExcluded;
 
 	TArray<UObject*> RootSetPackages;
 
-	EMode Mode;
+	int PassesLeft = 1;
+	bool bFirstPass = true;
 
+	bool NeedsAnotherPass() const { return PassesLeft > 0; }
+	void ProgressToNextPass() { bFirstPass = false; --PassesLeft; }
 public:
+	struct FRemovablePackagesInfo
+	{
+		TConstArrayView<FName> Assets;
+		TSet<FName> RemovablePackages;
+	};
+
 	FMemoryUsageReferenceProcessor();
 
-	bool Init(const TArray<FString>& PackageNames, FOutputDevice* ErrorOutput = GLog);
-	TArray<UObject*>& GetRootSet();
-	void HandleTokenStreamObjectReference(UE::GC::FWorkerContext& Context, const UObject* ReferencingObject, UObject*& Object, UE::GC::FTokenId TokenIndex, const EGCTokenType TokenType, bool bAllowReferenceElimination);
-	bool GetUnreachablePackages(TSet<FName>& OutUnreachablePackages);
+	void AddAssetSet(TConstArrayView<FName> PackageNames);
+	void PerformReachabilityAnalysis();
+	void GetUnreachablePackages(int AssetSetIdx, TSet<FName>& OutUnreachablePackages) const;
+	void GetUnreachablePackages(TArrayView<FRemovablePackagesInfo> OutUnreachablePackages) const;
 
-	void SetMode(EMode InMode) { Mode = InMode; }
+	// FSimpleReferenceProcessorBase interface
+	void HandleTokenStreamObjectReference(UE::GC::FWorkerContext& Context, const UObject* ReferencingObject, UObject*& Object, UE::GC::FTokenId TokenIndex, const EGCTokenType TokenType, bool bAllowReferenceElimination);
 };
 
-typedef UE::GC::TDefaultCollector<FMemoryUsageReferenceProcessor> FMemoryUsageReferenceCollector;
+using FMemoryUsageReferenceCollector = UE::GC::TDefaultCollector<FMemoryUsageReferenceProcessor>;
 
 bool GetLongName(FStringView ShortPackageName, FName& OutLongPackageName, FOutputDevice* ErrorOutput = GLog);
 bool GetLongNames(const TArray<FString>& PackageNames, TSet<FName>& OutLongPackageNames, FOutputDevice* ErrorOutput = GLog);
-bool GetLongNameAndDependencies(FStringView PackageName, FName& OutLongPackageName, TSet<FName>& OutDependencies, FOutputDevice* ErrorOutput = GLog);
 
-bool GetDependenciesCombined(const TArray<FString>& PackageNames, TSet<FName>& OutDependencies, FOutputDevice* ErrorOutput = GLog);
-bool GetDependenciesShared(const TArray<FString>& PackageNames, TSet<FName>& OutDependencies, FOutputDevice* ErrorOutput = GLog);
-
-bool PerformReachabilityAnalysis(FMemoryUsageReferenceProcessor& ReferenceProcessor, FOutputDevice* ErrorOutput = GLog);
+void GetDependenciesCombined(const TArray<FName>& PackageNames, TSet<FName>& OutDependencies);
+void GetDependenciesShared(const TArray<FName>& PackageNames, TSet<FName>& OutDependencies);
 
 // Get Packages that would be GC'd if PackagesToUnload were unloaded
-bool GetRemovablePackages(const TArray<FString>& PackagesToUnload, TSet<FName>& OutRemovablePackages, FOutputDevice* ErrorOutput = GLog);
+void GetRemovablePackages(const TArray<FName>& PackagesToUnload, TSet<FName>& OutRemovablePackages);
 
 // Get Packages that would not be GC'd if PackagesToUnload were unloaded
-bool GetUnremovablePackages(const TArray<FString>& PackagesToUnload, TSet<FName>& OutUnremovablePackages, FOutputDevice* ErrorOutput = GLog);
+void GetUnremovablePackages(const TArray<FName>& PackagesToUnload, TSet<FName>& OutUnremovablePackages);
+
+void GetRemovablePackagesBulk(TArrayView<FMemoryUsageReferenceProcessor::FRemovablePackagesInfo> Infos);
 
 void GetTransitiveDependencies(FName PackageName, TSet<FName>& OutDependencies);
 
-void SortPackagesBySize(const IMemoryUsageInfoProvider* MemoryUsageInfoProvider, const TSet<FName>& Packages, TMap<FName, uint64>& OutPackagesWithSize, FOutputDevice* ErrorOutput = GLog);
-void GetPackagesSize(const IMemoryUsageInfoProvider* MemoryUsageInfoProvider, const TSet<FName>& Packages, TMap<FName, uint64>& OutPackagesWithSize, FOutputDevice* ErrorOutput = GLog);
+void SortPackagesBySize(const IMemoryUsageInfoProvider* MemoryUsageInfoProvider, const TSet<FName>& Packages, TMap<FName, uint64>& OutPackagesWithSize);
+void GetPackagesSize(const IMemoryUsageInfoProvider* MemoryUsageInfoProvider, const TSet<FName>& Packages, TMap<FName, uint64>& OutPackagesWithSize);
 
 void RemoveNonExistentPackages(TMap<FName, uint64>& OutPackagesWithSize);
 void RemoveFilteredPackages(TMap<FName, uint64>& OutPackagesWithSize, const FString& AssetSubstring);
 
 void PrintTagsWithSize(FOutputDevice& Ar, const TMap<FName, uint64>& TagsWithSize, const TCHAR* Name, bool bTruncate = false, int32 Limit = -1, bool bCSV = false);
 
-}
+} // namespace MemoryUsageQueries::Internal
diff --git a/Engine/Plugins/MemoryUsageQueries/Source/MemoryUsageQueries/Public/MemoryUsageInfoProvider.h b/Engine/Plugins/MemoryUsageQueries/Source/MemoryUsageQueries/Public/MemoryUsageInfoProvider.h
--- a/Engine/Plugins/MemoryUsageQueries/Source/MemoryUsageQueries/Public/MemoryUsageInfoProvider.h
+++ b/Engine/Plugins/MemoryUsageQueries/Source/MemoryUsageQueries/Public/MemoryUsageInfoProvider.h
@@ -3,30 +3,24 @@
 #pragma once
 
 #include "Containers/Map.h"
-#include "HAL/LowLevelMemTracker.h"
 #include "CoreGlobals.h"
-#include "Misc/OutputDeviceRedirector.h"
 
 
 class MEMORYUSAGEQUERIES_API IMemoryUsageInfoProvider
 {
 public:
-	virtual uint64 GetAssetMemoryUsage(FName Asset, FOutputDevice* ErrorOutput) const = 0;
-	virtual uint64 GetAssetsMemoryUsage(const TSet<FName>& Assets, FOutputDevice* ErrorOutput) const = 0;
-	virtual void GetAllAssetsWithSize(TMap<FName, uint64>& OutAssets, FOutputDevice* ErrorOutput) const = 0;
-	virtual uint64 GetAssetsMemoryUsageWithSize(const TSet<FName>& Assets, TMap<FName, uint64>& OutSizes, FOutputDevice* ErrorOutput = GLog) const = 0;
-};
-
-class MEMORYUSAGEQUERIES_API FMemoryUsageInfoProviderLLM : public IMemoryUsageInfoProvider
-{
-public:
-	virtual uint64 GetAssetMemoryUsage(FName Asset, FOutputDevice* ErrorOutput = GLog) const override;
-	virtual uint64 GetAssetsMemoryUsage(const TSet<FName>& Assets, FOutputDevice* ErrorOutput = GLog) const override;
-	virtual void GetAllAssetsWithSize(TMap<FName, uint64>& OutAssets, FOutputDevice* ErrorOutput = GLog) const override;
-	virtual uint64 GetAssetsMemoryUsageWithSize(const TSet<FName>& Assets, TMap<FName, uint64>& OutSizes, FOutputDevice* ErrorOutput = GLog) const override;
-
+	UE_DEPRECATED(5.6, "Use version that do not accept FOutputDevice, check IsProviderAvailable() and do the logging yourself.")
+	virtual uint64 GetAssetMemoryUsage(FName Asset, FOutputDevice* ErrorOutput) const;
+	UE_DEPRECATED(5.6, "Use version that do not accept FOutputDevice, check IsProviderAvailable() and do the logging yourself.")
+	virtual uint64 GetAssetsMemoryUsage(const TSet<FName>& Assets, FOutputDevice* ErrorOutput) const;
+	UE_DEPRECATED(5.6, "Use version that do not accept FOutputDevice, check IsProviderAvailable() and do the logging yourself.")
+	virtual void GetAllAssetsWithSize(TMap<FName, uint64>& OutAssets, FOutputDevice* ErrorOutput) const;
+	UE_DEPRECATED(5.6, "Use version that do not accept FOutputDevice, check IsProviderAvailable() and do the logging yourself.")
+	virtual uint64 GetAssetsMemoryUsageWithSize(const TSet<FName>& Assets, TMap<FName, uint64>& OutSizes, FOutputDevice* ErrorOutput) const;
 
-#if ENABLE_LOW_LEVEL_MEM_TRACKER
-	virtual void GetFilteredTagsWithSize(TMap<FName, uint64>& OutTags, ELLMTracker Tracker, ELLMTagSet TagSet, TArray<FLLMTagSetAllocationFilter>& Filters, FOutputDevice* ErrorOutput = GLog) const;
-#endif
+	virtual bool IsProviderAvailable() const = 0;
+	virtual uint64 GetAssetMemoryUsage(FName Asset) const = 0;
+	virtual uint64 GetAssetsMemoryUsage(const TSet<FName>& Assets) const = 0;
+	virtual void GetAllAssetsWithSize(TMap<FName, uint64>& OutAssets) const = 0;
+	virtual uint64 GetAssetsMemoryUsageWithSize(const TSet<FName>& Assets, TMap<FName, uint64>& OutSizes) const = 0;
 };
diff --git a/Engine/Plugins/MemoryUsageQueries/Source/MemoryUsageQueries/Public/MemoryUsageQueries.h b/Engine/Plugins/MemoryUsageQueries/Source/MemoryUsageQueries/Public/MemoryUsageQueries.h
--- a/Engine/Plugins/MemoryUsageQueries/Source/MemoryUsageQueries/Public/MemoryUsageQueries.h
+++ b/Engine/Plugins/MemoryUsageQueries/Source/MemoryUsageQueries/Public/MemoryUsageQueries.h
@@ -3,6 +3,8 @@
 #pragma once
 
 #include "MemoryUsageInfoProvider.h"
+#include "HAL/LowLevelMemTracker.h"
+#include "Misc/OutputDeviceRedirector.h"
 
 struct FAutoCompleteCommand;
 
@@ -10,39 +12,56 @@ class FOutputDevice;
 
 namespace MemoryUsageQueries
 {
-
-extern MEMORYUSAGEQUERIES_API FMemoryUsageInfoProviderLLM MemoryUsageInfoProviderLLM;
-
-void RegisterConsoleAutoCompleteEntries(TArray<FAutoCompleteCommand>& AutoCompleteList);
-
-enum EDependencyType
+	enum EDependencyType
 {
 	EDep_All,
 	EDep_Removable,
 	EDep_NonRemovable
 };
-
 MEMORYUSAGEQUERIES_API IMemoryUsageInfoProvider* GetCurrentMemoryUsageInfoProvider();
+UE_DEPRECATED(5.6, "Use FName version instead, resolve user input using GetLongName if needed")
 MEMORYUSAGEQUERIES_API bool GetMemoryUsage(const IMemoryUsageInfoProvider* MemoryUsageInfoProvider, const FString& PackageName, uint64& OutExclusiveSize, uint64& OutInclusiveSize, FOutputDevice* ErrorOutput = GLog);
+UE_DEPRECATED(5.6, "Use TArray<FName> version instead, resolve user input using GetLongNames if needed")
 MEMORYUSAGEQUERIES_API bool GetMemoryUsageCombined(const IMemoryUsageInfoProvider* MemoryUsageInfoProvider, const TArray<FString>& PackageNames, uint64& OutTotalSize, FOutputDevice* ErrorOutput = GLog);
+UE_DEPRECATED(5.6, "Use TArray<FName> version instead, resolve user input using GetLongNames if needed")
 MEMORYUSAGEQUERIES_API bool GetMemoryUsageShared(const IMemoryUsageInfoProvider* MemoryUsageInfoProvider, const TArray<FString>& PackageNames, uint64& OutTotalSize, FOutputDevice* ErrorOutput = GLog);
+UE_DEPRECATED(5.6, "Use TArray<FName> version instead, resolve user input using GetLongNames if needed")
 MEMORYUSAGEQUERIES_API bool GetMemoryUsageUnique(const IMemoryUsageInfoProvider* MemoryUsageInfoProvider, const TArray<FString>& PackageNames, uint64& OutUniqueSize, FOutputDevice* ErrorOutput = GLog);
+UE_DEPRECATED(5.6, "Use TArray<FName> version instead, resolve user input using GetLongNames if needed")
 MEMORYUSAGEQUERIES_API bool GetMemoryUsageCommon(const IMemoryUsageInfoProvider* MemoryUsageInfoProvider, const TArray<FString>& PackageNames, uint64& OutCommonSize, FOutputDevice* ErrorOutput = GLog);
-						    
+
+UE_DEPRECATED(5.6, "Use FName version instead, resolve user input using GetLongName if needed")
 MEMORYUSAGEQUERIES_API bool GetDependenciesWithSize(const IMemoryUsageInfoProvider* MemoryUsageInfoProvider, const FString& PackageName, TMap<FName, uint64>& OutDependenciesWithSize, FOutputDevice* ErrorOutput = GLog);
+UE_DEPRECATED(5.6, "Use TArray<FName> version instead, resolve user input using GetLongNames if needed")
 MEMORYUSAGEQUERIES_API bool GetDependenciesWithSizeCombined(const IMemoryUsageInfoProvider* MemoryUsageInfoProvider, const TArray<FString>& PackageNames, TMap<FName, uint64>& OutDependenciesWithSize, FOutputDevice* ErrorOutput = GLog);
+UE_DEPRECATED(5.6, "Use TArray<FName> version instead, resolve user input using GetLongNames if needed")
 MEMORYUSAGEQUERIES_API bool GetDependenciesWithSizeShared(const IMemoryUsageInfoProvider* MemoryUsageInfoProvider, const TArray<FString>& PackageNames, TMap<FName, uint64>& OutDependenciesWithSize, FOutputDevice* ErrorOutput = GLog);
+UE_DEPRECATED(5.6, "Use TArray<FName> version instead, resolve user input using GetLongNames if needed")
 MEMORYUSAGEQUERIES_API bool GetDependenciesWithSizeUnique(const IMemoryUsageInfoProvider* MemoryUsageInfoProvider, const TArray<FString>& PackageNames, TMap<FName, uint64>& OutDependenciesWithSize, FOutputDevice* ErrorOutput = GLog);
+UE_DEPRECATED(5.6, "Use TArray<FName> version instead, resolve user input using GetLongNames if needed")
 MEMORYUSAGEQUERIES_API bool GetDependenciesWithSizeCommon(const IMemoryUsageInfoProvider* MemoryUsageInfoProvider, const TArray<FString>& PackageNames, TMap<FName, uint64>& OutDependenciesWithSize, FOutputDevice* ErrorOutput = GLog);
 
+MEMORYUSAGEQUERIES_API void GetMemoryUsage(const IMemoryUsageInfoProvider* MemoryUsageInfoProvider, const FName& PackageName, uint64& OutExclusiveSize, uint64& OutInclusiveSize);
+MEMORYUSAGEQUERIES_API void GetMemoryUsageCombined(const IMemoryUsageInfoProvider* MemoryUsageInfoProvider, const TArray<FName>& PackageNames, uint64& OutTotalSize);
+MEMORYUSAGEQUERIES_API void GetMemoryUsageShared(const IMemoryUsageInfoProvider* MemoryUsageInfoProvider, const TArray<FName>& PackageNames, uint64& OutTotalSize);
+MEMORYUSAGEQUERIES_API void GetMemoryUsageUnique(const IMemoryUsageInfoProvider* MemoryUsageInfoProvider, const TArray<FName>& PackageNames, uint64& OutUniqueSize);
+MEMORYUSAGEQUERIES_API void GetMemoryUsageCommon(const IMemoryUsageInfoProvider* MemoryUsageInfoProvider, const TArray<FName>& PackageNames, uint64& OutCommonSize);
+
+MEMORYUSAGEQUERIES_API void GetDependenciesWithSize(const IMemoryUsageInfoProvider* MemoryUsageInfoProvider, const FName& PackageName, TMap<FName, uint64>& OutDependenciesWithSize);
+MEMORYUSAGEQUERIES_API void GetDependenciesWithSizeCombined(const IMemoryUsageInfoProvider* MemoryUsageInfoProvider, const TArray<FName>& PackageNames, TMap<FName, uint64>& OutDependenciesWithSize);
+MEMORYUSAGEQUERIES_API void GetDependenciesWithSizeShared(const IMemoryUsageInfoProvider* MemoryUsageInfoProvider, const TArray<FName>& PackageNames, TMap<FName, uint64>& OutDependenciesWithSize);
+MEMORYUSAGEQUERIES_API void GetDependenciesWithSizeUnique(const IMemoryUsageInfoProvider* MemoryUsageInfoProvider, const TArray<FName>& PackageNames, TMap<FName, uint64>& OutDependenciesWithSize);
+MEMORYUSAGEQUERIES_API void GetDependenciesWithSizeCommon(const IMemoryUsageInfoProvider* MemoryUsageInfoProvider, const TArray<FName>& PackageNames, TMap<FName, uint64>& OutDependenciesWithSize);
+
 #if ENABLE_LOW_LEVEL_MEM_TRACKER
 
 MEMORYUSAGEQUERIES_API bool GetFilteredPackagesWithSize(TMap<FName, uint64>& OutPackagesWithSize, FName GroupName = NAME_None, FString AssetSubstring = FString(), FName ClassName = NAME_None, FOutputDevice* ErrorOutput = GLog);
 MEMORYUSAGEQUERIES_API bool GetFilteredPackagesCategorizedWithSize(TMap<FName, uint64>& OutPackagesWithSize, FName GroupName = NAME_None, FString AssetSubstring = FString(), FName ClassName = NAME_None, FName CategoryName = NAME_None, FOutputDevice* ErrorOutput = GLog);
 MEMORYUSAGEQUERIES_API bool GetFilteredClassesWithSize(TMap<FName, uint64>& OutClassesWithSize, FName GroupName = NAME_None, FString AssetName = FString(), FOutputDevice* ErrorOutput = GLog);
 MEMORYUSAGEQUERIES_API bool GetFilteredGroupsWithSize(TMap<FName, uint64>& OutGroupsWithSize, FString AssetName = FString(), FName ClassName = NAME_None, FOutputDevice* ErrorOutput = GLog);
-MEMORYUSAGEQUERIES_API bool GatherDependenciesForPackages(const IMemoryUsageInfoProvider* MemoryUsageInfoProvider, const TArray<FString>& PackageNames,
-	TMap<FName, uint64>& OutInternalDeps, TMap<FName, uint64>& OutExternalDeps, MemoryUsageQueries::EDependencyType DependencyType = MemoryUsageQueries::EDependencyType::EDep_All, FOutputDevice* ErrorOutput = GLog);
+UE_DEPRECATED(5.6, "Use TArray<FName> version instead, resolve user input using GetLongNames if needed")
+MEMORYUSAGEQUERIES_API bool GatherDependenciesForPackages(const IMemoryUsageInfoProvider* MemoryUsageInfoProvider, const TArray<FString>& PackageNames, TMap<FName, uint64>& OutInternalDeps, TMap<FName, uint64>& OutExternalDeps, MemoryUsageQueries::EDependencyType DependencyType = MemoryUsageQueries::EDependencyType::EDep_All, FOutputDevice* ErrorOutput = GLog);
+MEMORYUSAGEQUERIES_API void GatherDependenciesForPackages(const IMemoryUsageInfoProvider* MemoryUsageInfoProvider, const TArray<FName>& PackageNames, TMap<FName, uint64>& OutInternalDeps, TMap<FName, uint64>& OutExternalDeps, MemoryUsageQueries::EDependencyType DependencyType = MemoryUsageQueries::EDependencyType::EDep_All);
 #endif
 
 }
