[MQ] Partial backout of FMemoryUsageReferenceProcessor changes

#tests local profilego
#rnx

diff --git a/Engine/Plugins/MemoryUsageQueries/Source/MemoryUsageQueries/Private/MemoryUsageQueries.cpp b/Engine/Plugins/MemoryUsageQueries/Source/MemoryUsageQueries/Private/MemoryUsageQueries.cpp
--- a/Engine/Plugins/MemoryUsageQueries/Source/MemoryUsageQueries/Private/MemoryUsageQueries.cpp
+++ b/Engine/Plugins/MemoryUsageQueries/Source/MemoryUsageQueries/Private/MemoryUsageQueries.cpp
@@ -336,8 +336,13 @@ FMemoryUsageReferenceProcessor::FMemoryUsageReferenceProcessor()
 {
 	int32 Num = GUObjectArray.GetObjectArrayNum();
 
+	Excluded.Init(false, Num);
 	ReachableFull.Init(false, Num);
+	ReachableExcluded.Init(false, Num);
+}
 
+void FMemoryUsageReferenceProcessor::Init(const TArray<FName>& PackageNames)
+{
 	for (FRawObjectIterator It(true); It; ++It)
 	{
 		FUObjectItem* ObjectItem = *It;
@@ -367,25 +372,16 @@ FMemoryUsageReferenceProcessor::FMemoryUsageReferenceProcessor()
 	{
 		RootSetPackages.Add(FGCObject::GGCObjectReferencer);
 	}
-}
 
-void FMemoryUsageReferenceProcessor::AddAssetSet(TConstArrayView<FName> PackageNames)
-{
-	FAssetSet& AssetSet = AssetSets.Emplace_GetRef();
-
-	int Num = ReachableFull.Num();
-	AssetSet.Excluded.Init(false, Num);
-	AssetSet.ReachableExcluded.Init(false, Num);
-	
-	for (const FName& PackageName : PackageNames)
+	for (const auto& PackageName : PackageNames)
 	{
 		UPackage* Package = Cast<UPackage>(StaticFindObjectFast(UPackage::StaticClass(), nullptr, PackageName, /*bExactClass*/ true));
 		if (Package)
 		{
 			auto ExcludeObject = [&](UObject* ObjectToExclude) {
-				if (ObjectToExclude && GUObjectArray.ObjectToIndex(ObjectToExclude) < Num)
+				if (ObjectToExclude && GUObjectArray.ObjectToIndex(ObjectToExclude) < Excluded.Num())
 				{
-					AssetSet.Excluded[GUObjectArray.ObjectToIndex(ObjectToExclude)] = true;
+					Excluded[GUObjectArray.ObjectToIndex(ObjectToExclude)] = true;
 				}
 			};
 
@@ -412,22 +408,22 @@ void FMemoryUsageReferenceProcessor::AddAssetSet(TConstArrayView<FName> PackageN
 	}
 }
 
+TArray<UObject*>& FMemoryUsageReferenceProcessor::GetRootSet()
+{
+	return RootSetPackages;
+}
+
 void FMemoryUsageReferenceProcessor::HandleTokenStreamObjectReference(UE::GC::FWorkerContext& Context, const UObject* ReferencingObject, UObject*& Object, UE::GC::FTokenId TokenIndex, const EGCTokenType TokenType, bool bAllowReferenceElimination)
 {
 	FPermanentObjectPoolExtents PermanentPool;
-	if (Object == nullptr || PermanentPool.Contains(Object) || GUObjectArray.IsDisregardForGC(Object))
+	if (Object == nullptr || GUObjectArray.ObjectToIndex(Object) >= ReachableFull.Num() || PermanentPool.Contains(Object) || GUObjectArray.IsDisregardForGC(Object))
 	{
 		return;
 	}
 
 	int32 ObjectIndex = GUObjectArray.ObjectToIndex(Object);
 
-	if (ObjectIndex >= ReachableFull.Num())
-	{
-		return;
-	}
-
-	if (bFirstPass)
+	if (Mode == EMode::Full)
 	{
 		if (!ReachableFull[ObjectIndex])
 		{
@@ -435,33 +431,21 @@ void FMemoryUsageReferenceProcessor::HandleTokenStreamObjectReference(UE::GC::FW
 			Context.ObjectsToSerialize.Add<Options>(Object);
 		}
 	}
-	else
+	else if (Mode == EMode::Excluding)
 	{
-		FAssetSet& AssetSet = AssetSets[PassesLeft - 1];
-		if (!AssetSet.ReachableExcluded[ObjectIndex] && !AssetSet.Excluded[ObjectIndex])
+		if (!ReachableExcluded[ObjectIndex] && !Excluded[ObjectIndex])
 		{
-			AssetSet.ReachableExcluded[ObjectIndex] = true;
+			ReachableExcluded[ObjectIndex] = true;
 			Context.ObjectsToSerialize.Add<Options>(Object);
 		}
 	}
 }
 
-void FMemoryUsageReferenceProcessor::PerformReachabilityAnalysis()
-{
-	UE::GC::PadObjectArray(RootSetPackages);
-	for(;NeedsAnotherPass(); ProgressToNextPass())
-	{
-		FGCArrayStruct ArrayStruct;
-		ArrayStruct.SetInitialObjectsPrepadded(RootSetPackages);
-		CollectReferences<FMemoryUsageReferenceCollector, FMemoryUsageReferenceProcessor>(*this, ArrayStruct);
-	}
-}
-
-void FMemoryUsageReferenceProcessor::GetUnreachablePackages(int AssetSetIdx, TSet<FName>& OutUnreachablePackages) const
+void FMemoryUsageReferenceProcessor::GetUnreachablePackages(TSet<FName>& OutUnreachablePackages)
 {
 	for (int i = 0; i < ReachableFull.Num(); i++)
 	{
-		if (ReachableFull[i] && !AssetSets[AssetSetIdx].ReachableExcluded[i])
+		if (ReachableFull[i] && !ReachableExcluded[i])
 		{
 			UObject* Obj = static_cast<UObject*>(GUObjectArray.IndexToObjectUnsafeForGC(i)->Object);
 			if (Obj && Obj->IsA<UPackage>())
@@ -472,36 +456,6 @@ void FMemoryUsageReferenceProcessor::GetUnreachablePackages(int AssetSetIdx, TSe
 	}
 }
 
-void FMemoryUsageReferenceProcessor::GetUnreachablePackages(TArrayView<FRemovablePackagesInfo> OutUnreachablePackages) const
-{
-	int Num = AssetSets.Num();
-	if(ensure(OutUnreachablePackages.Num() == Num))
-	{
-		for (int i = 0; i < ReachableFull.Num(); i++)
-		{
-			if (!ReachableFull[i])
-			{
-				continue;
-			}
-			UObject* Obj = nullptr;
-			for (int AssetSetIdx = 0; AssetSetIdx < Num; ++AssetSetIdx)
-			{
-				const FAssetSet& AssetSet = AssetSets[AssetSetIdx];
-				if (!AssetSet.ReachableExcluded[i])
-				{
-					if (!Obj)
-					{
-						Obj = static_cast<UObject*>(GUObjectArray.IndexToObjectUnsafeForGC(i)->Object);
-					}
-					if (Obj && Obj->IsA<UPackage>())
-					{
-						OutUnreachablePackages[AssetSetIdx].RemovablePackages.Add(Obj->GetFName());
-					}
-				}
-			}
-		}
-	}
-}
 
 static FAutoConsoleVariable CVarMemQueryUsePackageStore(
 	TEXT("MemQuery.UsePackageStore"), true,
@@ -952,25 +906,41 @@ void GetDependenciesShared(const TArray<FName>& PackageNames, TSet<FName>& OutDe
 	}
 }
 
+void PerformReachabilityAnalysis(FMemoryUsageReferenceProcessor& ReferenceProcessor)
+{
+	{
+		FGCArrayStruct ArrayStruct;
+		ArrayStruct.SetInitialObjectsUnpadded(ReferenceProcessor.GetRootSet());
+		ReferenceProcessor.SetMode(FMemoryUsageReferenceProcessor::Full);
+		CollectReferences<FMemoryUsageReferenceCollector, FMemoryUsageReferenceProcessor>(ReferenceProcessor, ArrayStruct);
+	}
+
+	{
+		FGCArrayStruct ArrayStruct;
+		ArrayStruct.SetInitialObjectsUnpadded(ReferenceProcessor.GetRootSet());
+		ReferenceProcessor.SetMode(FMemoryUsageReferenceProcessor::Excluding);
+		CollectReferences<FMemoryUsageReferenceCollector, FMemoryUsageReferenceProcessor>(ReferenceProcessor, ArrayStruct);
+	}
+}
+
 void GetRemovablePackages(const TArray<FName>& PackagesToUnload, TSet<FName>& OutRemovablePackages)
 {
 	FMemoryUsageReferenceProcessor ReferenceProcessor;
-	ReferenceProcessor.AddAssetSet(PackagesToUnload);
-	ReferenceProcessor.PerformReachabilityAnalysis();
-	ReferenceProcessor.GetUnreachablePackages(0, OutRemovablePackages);
+	ReferenceProcessor.Init(PackagesToUnload);
+	PerformReachabilityAnalysis(ReferenceProcessor);
+	ReferenceProcessor.GetUnreachablePackages(OutRemovablePackages);
 }
 
 void GetUnremovablePackages(const TArray<FName>& PackagesToUnload, TSet<FName>& OutUnremovablePackages)
 {
 	FMemoryUsageReferenceProcessor ReferenceProcessor;
-	ReferenceProcessor.AddAssetSet(PackagesToUnload);
-	ReferenceProcessor.PerformReachabilityAnalysis();
-
+	ReferenceProcessor.Init(PackagesToUnload);
+	PerformReachabilityAnalysis(ReferenceProcessor);
 	TSet<FName> UnreachablePackages;
-	ReferenceProcessor.GetUnreachablePackages(0, UnreachablePackages);
+	ReferenceProcessor.GetUnreachablePackages(UnreachablePackages);
 
 	TSet<FName> Dependencies;
-	GetDependenciesCombined(PackagesToUnload, Dependencies);
+	Internal::GetDependenciesCombined(PackagesToUnload, Dependencies);
 
 	for (const auto& Package : Dependencies)
 	{
@@ -981,17 +951,6 @@ void GetUnremovablePackages(const TArray<FName>& PackagesToUnload, TSet<FName>&
 	}
 }
 
-void GetRemovablePackagesBulk(TArrayView<FMemoryUsageReferenceProcessor::FRemovablePackagesInfo> Infos)
-{
-	FMemoryUsageReferenceProcessor ReferenceProcessor;
-	for (const FMemoryUsageReferenceProcessor::FRemovablePackagesInfo& Info : Infos)
-	{
-		ReferenceProcessor.AddAssetSet(Info.Assets);
-	}
-	ReferenceProcessor.PerformReachabilityAnalysis();
-	ReferenceProcessor.GetUnreachablePackages(Infos);
-}
-
 static void GetTransitiveDependenciesAssetRegistry(FName PackageName, TSet<FName>& OutDependencies)
 {
 	FAssetRegistryModule& AssetRegistryModule = GetAssetRegistryModule();
diff --git a/Engine/Plugins/MemoryUsageQueries/Source/MemoryUsageQueries/Private/MemoryUsageQueriesPrivate.h b/Engine/Plugins/MemoryUsageQueries/Source/MemoryUsageQueries/Private/MemoryUsageQueriesPrivate.h
--- a/Engine/Plugins/MemoryUsageQueries/Source/MemoryUsageQueries/Private/MemoryUsageQueriesPrivate.h
+++ b/Engine/Plugins/MemoryUsageQueries/Source/MemoryUsageQueries/Private/MemoryUsageQueriesPrivate.h
@@ -10,39 +10,32 @@ namespace MemoryUsageQueries::Internal
 
 class FMemoryUsageReferenceProcessor : public FSimpleReferenceProcessorBase
 {
-private:
-	struct FAssetSet
+public:
+	enum EMode
 	{
-		TBitArray<> Excluded;
-		TBitArray<> ReachableExcluded;
+		Full,
+		Excluding
 	};
-	TArray<FAssetSet> AssetSets;
+
+private:
+	TBitArray<> Excluded;
 
 	TBitArray<> ReachableFull;
+	TBitArray<> ReachableExcluded;
 
 	TArray<UObject*> RootSetPackages;
 
-	int PassesLeft = 1;
-	bool bFirstPass = true;
+	EMode Mode;
 
-	bool NeedsAnotherPass() const { return PassesLeft > 0; }
-	void ProgressToNextPass() { bFirstPass = false; --PassesLeft; }
 public:
-	struct FRemovablePackagesInfo
-	{
-		TConstArrayView<FName> Assets;
-		TSet<FName> RemovablePackages;
-	};
-
 	FMemoryUsageReferenceProcessor();
 
-	void AddAssetSet(TConstArrayView<FName> PackageNames);
-	void PerformReachabilityAnalysis();
-	void GetUnreachablePackages(int AssetSetIdx, TSet<FName>& OutUnreachablePackages) const;
-	void GetUnreachablePackages(TArrayView<FRemovablePackagesInfo> OutUnreachablePackages) const;
-
-	// FSimpleReferenceProcessorBase interface
+	void Init(const TArray<FName>& PackageNames);
+	TArray<UObject*>& GetRootSet();
 	void HandleTokenStreamObjectReference(UE::GC::FWorkerContext& Context, const UObject* ReferencingObject, UObject*& Object, UE::GC::FTokenId TokenIndex, const EGCTokenType TokenType, bool bAllowReferenceElimination);
+	void GetUnreachablePackages(TSet<FName>& OutUnreachablePackages);
+
+	void SetMode(EMode InMode) { Mode = InMode; }
 };
 
 using FMemoryUsageReferenceCollector = UE::GC::TDefaultCollector<FMemoryUsageReferenceProcessor>;
@@ -59,8 +52,6 @@ void GetRemovablePackages(const TArray<FName>& PackagesToUnload, TSet<FName>& Ou
 // Get Packages that would not be GC'd if PackagesToUnload were unloaded
 void GetUnremovablePackages(const TArray<FName>& PackagesToUnload, TSet<FName>& OutUnremovablePackages);
 
-void GetRemovablePackagesBulk(TArrayView<FMemoryUsageReferenceProcessor::FRemovablePackagesInfo> Infos);
-
 void GetTransitiveDependencies(FName PackageName, TSet<FName>& OutDependencies);
 
 void SortPackagesBySize(const IMemoryUsageInfoProvider* MemoryUsageInfoProvider, const TSet<FName>& Packages, TMap<FName, uint64>& OutPackagesWithSize);
