Added missing Greater.h include.

diff --git a/Engine/Plugins/MemoryUsageQueries/Source/MemoryUsageQueries/Private/MemoryUsageQueriesCommands.cpp b/Engine/Plugins/MemoryUsageQueries/Source/MemoryUsageQueries/Private/MemoryUsageQueriesCommands.cpp
--- a/Engine/Plugins/MemoryUsageQueries/Source/MemoryUsageQueries/Private/MemoryUsageQueriesCommands.cpp
+++ b/Engine/Plugins/MemoryUsageQueries/Source/MemoryUsageQueries/Private/MemoryUsageQueriesCommands.cpp
@@ -9,6 +9,7 @@
 #include "Misc/OutputDeviceArchiveWrapper.h"
 #include "Misc/WildcardString.h"
 #include "HAL/FileManager.h"
+#include "Templates/Greater.h"
 
 struct FAssetMemoryBreakdown
 {
