Fixed occurrences of using LLM_SCOPE instead of LLM_SCOPE_BYNAME. This fixes unnecessary construction of FNames on each call (2x).

#Synced-CL 50179990

diff --git a/Engine/Plugins/MemoryUsageQueries/Source/MemoryUsageQueries/Private/MemoryUsageQueries.cpp b/Engine/Plugins/MemoryUsageQueries/Source/MemoryUsageQueries/Private/MemoryUsageQueries.cpp
--- a/Engine/Plugins/MemoryUsageQueries/Source/MemoryUsageQueries/Private/MemoryUsageQueries.cpp
+++ b/Engine/Plugins/MemoryUsageQueries/Source/MemoryUsageQueries/Private/MemoryUsageQueries.cpp
@@ -510,7 +510,7 @@ class FPackageStoreLazyDatabase final
 				{
 					FPackageId PackageId = FPackageId::FromName(Data.PackageName);
 					{
-						LLM_SCOPE(TEXT("MemoryUsageQueries"));
+						LLM_SCOPE_BYNAME("MemoryUsageQueries");
 						Database.FindOrAdd(PackageId, Data.PackageName);
 						Names.Add(Data.PackageName);
 					}
@@ -540,7 +540,7 @@ class FPackageStoreLazyDatabase final
 				{
 					FPackageId PackageId = FPackageId::FromName(PackageName);
 					{
-						LLM_SCOPE(TEXT("MemoryUsageQueries"));
+						LLM_SCOPE_BYNAME("MemoryUsageQueries");
 						Database.FindOrAdd(PackageId, PackageName);
 						Names.Add(PackageName);
 					}
@@ -676,7 +676,7 @@ class FPackageDependenciesLazyDatabase final
 
 	bool InsertPackage(const FPackageId RootPackageId)
 	{
-		LLM_SCOPE(TEXT("MemoryUsageQueries"));
+		LLM_SCOPE_BYNAME("MemoryUsageQueries");
 
 		TArray<FPackageId, TInlineAllocator<2048>> Stack;
 
diff --git a/Engine/Plugins/Runtime/GameplayAbilities/Source/GameplayAbilities/Private/AbilitySystemGlobals.cpp b/Engine/Plugins/Runtime/GameplayAbilities/Source/GameplayAbilities/Private/AbilitySystemGlobals.cpp
--- a/Engine/Plugins/Runtime/GameplayAbilities/Source/GameplayAbilities/Private/AbilitySystemGlobals.cpp
+++ b/Engine/Plugins/Runtime/GameplayAbilities/Source/GameplayAbilities/Private/AbilitySystemGlobals.cpp
@@ -69,7 +69,7 @@ void UAbilitySystemGlobals::InitGlobalData()
 	}
 	bInitialized = true;
 
-	LLM_SCOPE(TEXT("AbilitySystem"));
+	LLM_SCOPE_BYNAME(TEXT("AbilitySystem"));
 	GetGlobalCurveTable();
 	GetGlobalAttributeMetaDataTable();
 	
diff --git a/Engine/Source/Editor/UnrealEd/Private/Cooker/CookConfigAccessTracker.cpp b/Engine/Source/Editor/UnrealEd/Private/Cooker/CookConfigAccessTracker.cpp
--- a/Engine/Source/Editor/UnrealEd/Private/Cooker/CookConfigAccessTracker.cpp
+++ b/Engine/Source/Editor/UnrealEd/Private/Cooker/CookConfigAccessTracker.cpp
@@ -332,7 +332,7 @@ void FCookConfigAccessTracker::StaticOnConfigValueRead(UE::ConfigAccessTracking:
 	}
 #endif
 
-	LLM_SCOPE_BYNAME(TEXTVIEW("ConfigAccessTracking"));
+	LLM_SCOPE_BYNAME("ConfigAccessTracking");
 	FConfigAccessData AccessData(ConfigFile->LoadType, ConfigPlatform, FileName, SectionName, ValueName, RequestedPlatform);
 	uint32 ReferencerHash = GetTypeHash(Referencer);
 	uint32 AccessDataHash = GetTypeHash(AccessData);
diff --git a/Engine/Source/Editor/UnrealEd/Private/Cooker/PackageBuildDependencyTracker.cpp b/Engine/Source/Editor/UnrealEd/Private/Cooker/PackageBuildDependencyTracker.cpp
--- a/Engine/Source/Editor/UnrealEd/Private/Cooker/PackageBuildDependencyTracker.cpp
+++ b/Engine/Source/Editor/UnrealEd/Private/Cooker/PackageBuildDependencyTracker.cpp
@@ -175,7 +175,7 @@ void FPackageBuildDependencyTracker::StaticOnObjectHandleRead(const TArrayView<c
 				*Referencer.ToString(), *Referenced.ToString());
 		}
 
-		LLM_SCOPE_BYNAME(TEXTVIEW("PackageBuildDependencyTracker"));
+		LLM_SCOPE_BYNAME("PackageBuildDependencyTracker");
 
 		FBuildDependencyAccessData AccessData{ Referenced, AccumulatedScopeData->TargetPlatform };
 		UE::TUniqueLock RecordsScopeLock(Singleton.RecordsLock);
